ElectraSubtitles: Created a plugin to handle subtitles for ElectraPlayer

diff --git a/Engine/Plugins/Media/ElectraSubtitles/ElectraSubtitles.uplugin b/Engine/Plugins/Media/ElectraSubtitles/ElectraSubtitles.uplugin
new file mode 100644
--- /dev/null
+++ b/Engine/Plugins/Media/ElectraSubtitles/ElectraSubtitles.uplugin
@@ -0,0 +1,32 @@
+{
+	"FileVersion": 1,
+	"Version": 1,
+	"VersionName": "1.0",
+	"FriendlyName": "Electra Player Subtitle Module",
+	"Description": "Subtitle Decoder Module for Electra Player Media Playback",
+	"Category": "Media Players",
+	"CreatedBy": "Epic Games, Inc.",
+	"CreatedByURL": "http://epicgames.com",
+	"DocsURL": "https://docs.unrealengine.com/en-US/Engine/MediaFramework/Overview",
+	"MarketplaceURL": "",
+	"SupportURL": "",
+	"EnabledByDefault": false,
+	"CanContainContent": false,
+	"IsBetaVersion": false,
+	"Installed": false,
+	"Modules": [
+		{
+			"Name": "ElectraSubtitles",
+			"Type": "RuntimeNoCommandlet",
+			"LoadingPhase": "PostEngineInit",
+			"WhitelistPlatforms": [ "Win64", "Mac", "IOS", "Android" ],
+			"BlacklistTargets": [ "Server" ]
+		}
+	],
+	"Plugins": [
+		{
+			"Name": "ElectraUtil",
+			"Enabled": true
+		}
+	]
+}
\ No newline at end of file
diff --git a/Engine/Plugins/Media/ElectraSubtitles/Source/ElectraSubtitles/ElectraSubtitles.Build.cs b/Engine/Plugins/Media/ElectraSubtitles/Source/ElectraSubtitles/ElectraSubtitles.Build.cs
new file mode 100644
--- /dev/null
+++ b/Engine/Plugins/Media/ElectraSubtitles/Source/ElectraSubtitles/ElectraSubtitles.Build.cs
@@ -0,0 +1,20 @@
+// Copyright Epic Games, Inc. All Rights Reserved.
+
+namespace UnrealBuildTool.Rules
+{
+	public class ElectraSubtitles : ModuleRules
+	{
+		public ElectraSubtitles(ReadOnlyTargetRules Target) : base(Target)
+		{
+            PCHUsage = PCHUsageMode.NoPCHs;
+			PrivateDependencyModuleNames.AddRange(
+				new string[]
+				{
+					"Core",
+					"ElectraBase",
+					"ElectraSamples",
+					"XmlParser"
+				});
+		}
+	}
+}
diff --git a/Engine/Plugins/Media/ElectraSubtitles/Source/ElectraSubtitles/Private/ElectraSubtitleDecoder.h b/Engine/Plugins/Media/ElectraSubtitles/Source/ElectraSubtitles/Private/ElectraSubtitleDecoder.h
new file mode 100644
--- /dev/null
+++ b/Engine/Plugins/Media/ElectraSubtitles/Source/ElectraSubtitles/Private/ElectraSubtitleDecoder.h
@@ -0,0 +1,8 @@
+// Copyright Epic Games, Inc. All Rights Reserved.
+
+#pragma once
+
+#include "IElectraSubtitleDecoder.h"
+
+
+
diff --git a/Engine/Plugins/Media/ElectraSubtitles/Source/ElectraSubtitles/Private/ElectraSubtitleDecoderFactory.h b/Engine/Plugins/Media/ElectraSubtitles/Source/ElectraSubtitles/Private/ElectraSubtitleDecoderFactory.h
new file mode 100644
--- /dev/null
+++ b/Engine/Plugins/Media/ElectraSubtitles/Source/ElectraSubtitles/Private/ElectraSubtitleDecoderFactory.h
@@ -0,0 +1,30 @@
+// Copyright Epic Games, Inc. All Rights Reserved.
+
+#pragma once
+
+#include "CoreMinimal.h"
+
+
+class IElectraSubtitleDecoderFactory
+{
+public:
+	virtual ~IElectraSubtitleDecoderFactory() = default;
+	virtual TSharedPtr<IElectraSubtitleDecoder, ESPMode::ThreadSafe> CreateDecoder(const FString& SubtitleCodecName) = 0;
+};
+
+
+class IElectraSubtitleDecoderFactoryRegistry
+{
+public:
+	virtual ~IElectraSubtitleDecoderFactoryRegistry() = default;
+
+	struct FCodecInfo
+	{
+		FString	CodecName;
+		int32 Priority = 0;
+	};
+
+	virtual void AddDecoderFactory(const TArray<FCodecInfo>& InCodecInformation, IElectraSubtitleDecoderFactory* InDecoderFactory) = 0;
+
+};
+
diff --git a/Engine/Plugins/Media/ElectraSubtitles/Source/ElectraSubtitles/Private/ElectraSubtitleModule.cpp b/Engine/Plugins/Media/ElectraSubtitles/Source/ElectraSubtitles/Private/ElectraSubtitleModule.cpp
new file mode 100644
--- /dev/null
+++ b/Engine/Plugins/Media/ElectraSubtitles/Source/ElectraSubtitles/Private/ElectraSubtitleModule.cpp
@@ -0,0 +1,138 @@
+// Copyright Epic Games, Inc. All Rights Reserved.
+
+#include "CoreMinimal.h"
+#include "Misc/ScopeLock.h"
+
+#include "IElectraSubtitleModule.h"
+#include "ElectraSubtitleModule.h"
+#include "ElectraSubtitleDecoderFactory.h"
+
+#include "tx3g/ElectraSubtitleDecoder_TX3G.h"
+
+#define LOCTEXT_NAMESPACE "ElectraSubtitlesModule"
+
+DEFINE_LOG_CATEGORY(LogElectraSubtitles);
+
+// -----------------------------------------------------------------------------------------------------------------------------------
+
+class FElectraSubtitleModularFeature : public IElectraSubtitleModularFeature, public IElectraSubtitleDecoderFactoryRegistry
+{
+public:
+	virtual ~FElectraSubtitleModularFeature();
+
+	//-------------------------------------------------------------------------
+	// Methods from IElectraSubtitleModularFeature
+	//
+	virtual bool SupportsFormat(const FString& SubtitleCodecName) const override;
+
+	virtual void GetSupportedFormats(TArray<FString>& OutSupportedCodecNames) const override;
+
+	virtual int32 GetPriorityForFormat(const FString& SubtitleCodecName) const override;
+
+	virtual TSharedPtr<IElectraSubtitleDecoder, ESPMode::ThreadSafe> CreateDecoderForFormat(const FString& SubtitleCodecName) override;
+
+	//-------------------------------------------------------------------------
+	// Methods from IElectraSubtitleDecoderFactoryRegistry
+	//
+	virtual void AddDecoderFactory(const TArray<FCodecInfo>& InCodecInformation, IElectraSubtitleDecoderFactory* InDecoderFactory) override;
+
+private:
+	struct FSupportedCodec
+	{
+		int32 Priority = -1;
+		IElectraSubtitleDecoderFactory* Factory = nullptr;
+
+	};
+	mutable FCriticalSection Lock;
+	TMap<FString, FSupportedCodec> SupportedCodecs;
+};
+
+// -----------------------------------------------------------------------------------------------------------------------------------
+
+class FElectraSubtitlesModule : public IElectraSubtitlesModule
+{
+public:
+	virtual void StartupModule() override
+	{
+		// Register the codecs this plugin provides.
+		FElectraSubtitleDecoderTX3G::RegisterCodecs(ModularFeature);
+
+		IModularFeatures::Get().RegisterModularFeature(GetModularFeatureName(), &ModularFeature);
+	}
+
+	virtual void ShutdownModule() override
+	{
+		IModularFeatures::Get().UnregisterModularFeature(GetModularFeatureName(), &ModularFeature);
+	}
+	
+	virtual bool SupportsDynamicReloading() override
+	{
+		return false;
+	}
+
+private:
+	FElectraSubtitleModularFeature ModularFeature;
+};
+
+IMPLEMENT_MODULE(FElectraSubtitlesModule, ElectraSubtitles);
+
+// -----------------------------------------------------------------------------------------------------------------------------------
+
+void FElectraSubtitleModularFeature::AddDecoderFactory(const TArray<FCodecInfo>& InCodecInformation, IElectraSubtitleDecoderFactory* InDecoderFactory)
+{
+	FScopeLock lock(&Lock);
+	// Every codec may claim support for more than one format, so we add it in multiple times, once per format.
+	// Also, if there is already a registered format we compare the priorities and take only the one with a higher priority.
+	// In case of a tie in priority we take the last one added.
+	for(auto &ci : InCodecInformation)
+	{
+		bool bAdd = true;
+		for(auto &Codec : SupportedCodecs)
+		{
+			if (Codec.Key.Equals(ci.CodecName) && Codec.Value.Priority > ci.Priority)
+			{
+				bAdd = false;
+				break;
+			}
+		}
+		if (bAdd)
+		{
+			SupportedCodecs.Emplace(ci.CodecName, FSupportedCodec( {ci.Priority, InDecoderFactory} ));
+		}
+	}
+}
+
+
+FElectraSubtitleModularFeature::~FElectraSubtitleModularFeature()
+{
+}
+
+bool FElectraSubtitleModularFeature::SupportsFormat(const FString& SubtitleCodecName) const
+{
+	FScopeLock lock(&Lock);
+	return SupportedCodecs.Contains(SubtitleCodecName);
+}
+
+void FElectraSubtitleModularFeature::GetSupportedFormats(TArray<FString>& OutSupportedCodecNames) const
+{
+	FScopeLock lock(&Lock);
+	SupportedCodecs.GenerateKeyArray(OutSupportedCodecNames);
+}
+
+int32 FElectraSubtitleModularFeature::GetPriorityForFormat(const FString& SubtitleCodecName) const
+{
+	FScopeLock lock(&Lock);
+	const FSupportedCodec* Codec = SupportedCodecs.Find(SubtitleCodecName);
+	return Codec ? Codec->Priority : -1;
+}
+
+TSharedPtr<IElectraSubtitleDecoder, ESPMode::ThreadSafe> FElectraSubtitleModularFeature::CreateDecoderForFormat(const FString& SubtitleCodecName)
+{
+	FScopeLock lock(&Lock);
+	const FSupportedCodec* Codec = SupportedCodecs.Find(SubtitleCodecName);
+	return Codec ? Codec->Factory->CreateDecoder(SubtitleCodecName) : nullptr;
+}
+
+
+#undef LOCTEXT_NAMESPACE
+
diff --git a/Engine/Plugins/Media/ElectraSubtitles/Source/ElectraSubtitles/Private/ElectraSubtitleModule.h b/Engine/Plugins/Media/ElectraSubtitles/Source/ElectraSubtitles/Private/ElectraSubtitleModule.h
new file mode 100644
--- /dev/null
+++ b/Engine/Plugins/Media/ElectraSubtitles/Source/ElectraSubtitles/Private/ElectraSubtitleModule.h
@@ -0,0 +1,8 @@
+// Copyright Epic Games, Inc. All Rights Reserved.
+
+#pragma once
+
+#include <CoreMinimal.h>
+
+DECLARE_LOG_CATEGORY_EXTERN(LogElectraSubtitles, Log, All);
+
diff --git a/Engine/Plugins/Media/ElectraSubtitles/Source/ElectraSubtitles/Private/tx3g/ElectraSubtitleDecoder_TX3G.cpp b/Engine/Plugins/Media/ElectraSubtitles/Source/ElectraSubtitles/Private/tx3g/ElectraSubtitleDecoder_TX3G.cpp
new file mode 100644
--- /dev/null
+++ b/Engine/Plugins/Media/ElectraSubtitles/Source/ElectraSubtitles/Private/tx3g/ElectraSubtitleDecoder_TX3G.cpp
@@ -0,0 +1,480 @@
+// Copyright Epic Games, Inc. All Rights Reserved.
+
+#include "tx3g/ElectraSubtitleDecoder_TX3G.h"
+#include "Containers/StringConv.h"
+
+namespace ElectraSubtitleDecoderTX3GUtils
+{
+
+#if !PLATFORM_LITTLE_ENDIAN
+	static inline uint8  GetFromBigEndian(uint8 value)		{ return value; }
+	static inline int8   GetFromBigEndian(int8 value)		{ return value; }
+	static inline uint16 GetFromBigEndian(uint16 value)		{ return value; }
+	static inline int16 GetFromBigEndian(int16 value)		{ return value; }
+	static inline int32 GetFromBigEndian(int32 value)		{ return value; }
+	static inline uint32 GetFromBigEndian(uint32 value)		{ return value; }
+	static inline int64 GetFromBigEndian(int64 value)		{ return value; }
+	static inline uint64 GetFromBigEndian(uint64 value)		{ return value; }
+#else
+	static inline uint16 EndianSwap(uint16 value)			{ return (value >> 8) | (value << 8); }
+	static inline int16 EndianSwap(int16 value)				{ return int16(EndianSwap(uint16(value))); }
+	static inline uint32 EndianSwap(uint32 value)			{ return (value << 24) | ((value & 0xff00) << 8) | ((value >> 8) & 0xff00) | (value >> 24); }
+	static inline int32 EndianSwap(int32 value)				{ return int32(EndianSwap(uint32(value))); }
+	static inline uint64 EndianSwap(uint64 value)			{ return (uint64(EndianSwap(uint32(value & 0xffffffffU))) << 32) | uint64(EndianSwap(uint32(value >> 32))); }
+	static inline int64 EndianSwap(int64 value)				{ return int64(EndianSwap(uint64(value)));}
+	static inline uint8  GetFromBigEndian(uint8 value)		{ return value; }
+	static inline int8   GetFromBigEndian(int8 value)		{ return value; }
+	static inline uint16 GetFromBigEndian(uint16 value)		{ return EndianSwap(value); }
+	static inline int16 GetFromBigEndian(int16 value)		{ return EndianSwap(value); }
+	static inline int32 GetFromBigEndian(int32 value)		{ return EndianSwap(value); }
+	static inline uint32 GetFromBigEndian(uint32 value)		{ return EndianSwap(value); }
+	static inline int64 GetFromBigEndian(int64 value)		{ return EndianSwap(value); }
+	static inline uint64 GetFromBigEndian(uint64 value)		{ return EndianSwap(value); }
+#endif
+
+	template <typename T>
+	T ValueFromBigEndian(const T value)
+	{
+		return GetFromBigEndian(value);
+	}
+
+
+	class FDataReaderMP4
+	{
+	public:
+		FDataReaderMP4(const TArray<uint8>& InDataBufferToReadFrom) : DataBufferRef(InDataBufferToReadFrom)
+		{}
+
+		int32 GetCurrentOffset() const
+		{ return CurrentOffset;	}
+
+		int32 GetNumBytesRemaining() const
+		{ return DataBufferRef.Num() - GetCurrentOffset(); }
+
+		template <typename T>
+		bool Read(T& value)
+		{
+			T Temp = 0;
+			int64 NumRead = ReadData(&Temp, sizeof(T));
+			if (NumRead == sizeof(T))
+			{
+				value = ElectraSubtitleDecoderTX3GUtils::ValueFromBigEndian(Temp);
+				return true;
+			}
+			return false;
+		}
+
+		bool ReadString(FString& OutString, uint16 NumBytes)
+		{
+			OutString.Empty();
+			if (NumBytes == 0)
+			{
+				return true;
+			}
+			TArray<uint8> Buf;
+			Buf.AddUninitialized(NumBytes);
+			if (ReadBytes(Buf.GetData(), NumBytes))
+			{
+				// Check for UTF16 BOM
+				if (NumBytes >= 2 && ((Buf[0] == 0xff && Buf[1] == 0xfe) || (Buf[0] == 0xfe && Buf[1] == 0xff)))
+				{
+					UE_LOG(LogElectraSubtitles, Error, TEXT("TX3G uses UTF16 which is not supported"));
+					return false;
+				}
+				FUTF8ToTCHAR cnv((const ANSICHAR*)Buf.GetData(), NumBytes);
+				OutString = FString(cnv.Length(), cnv.Get());
+				return true;
+			}
+			return false;
+		}
+
+		bool ReadBytes(void* Buffer, int32 NumBytes)
+		{
+			int32 NumRead = ReadData(Buffer, NumBytes);
+			return NumRead == NumBytes;
+		}
+
+#define MAKE_BOX_ATOM(a,b,c,d) (uint32)((uint32)a << 24) | ((uint32)b << 16) | ((uint32)c << 8) | ((uint32)d)
+		static const uint32 BoxType_styl = MAKE_BOX_ATOM('s', 't', 'y', 'l');
+		static const uint32 BoxType_hlit = MAKE_BOX_ATOM('h', 'l', 'i', 't');
+		static const uint32 BoxType_hclr = MAKE_BOX_ATOM('h', 'c', 'l', 'r');
+		static const uint32 BoxType_krok = MAKE_BOX_ATOM('k', 'r', 'o', 'k');
+		static const uint32 BoxType_dlay = MAKE_BOX_ATOM('d', 'l', 'a', 'y');
+		static const uint32 BoxType_href = MAKE_BOX_ATOM('h', 'r', 'e', 'f');
+		static const uint32 BoxType_tbox = MAKE_BOX_ATOM('t', 'b', 'o', 'x');
+		static const uint32 BoxType_blnk = MAKE_BOX_ATOM('b', 'l', 'n', 'k');
+		static const uint32 BoxType_twrp = MAKE_BOX_ATOM('t', 'w', 'r', 'p');
+#undef MAKE_BOX_ATOM
+	private:
+		int32 ReadData(void* IntoBuffer, int32 NumBytesToRead)
+		{
+			if (NumBytesToRead == 0)
+			{
+				return 0;
+			}
+			int32 NumAvail = DataBufferRef.Num() - CurrentOffset;
+			if (NumAvail >= NumBytesToRead)
+			{
+				if (IntoBuffer)
+				{
+					FMemory::Memcpy(IntoBuffer, DataBufferRef.GetData() + CurrentOffset, NumBytesToRead);
+				}
+				CurrentOffset += NumBytesToRead;
+				return NumBytesToRead;
+			}
+			return -1;
+		}
+
+		const TArray<uint8>& DataBufferRef;
+		int32 CurrentOffset = 0;
+	};
+
+}
+
+
+
+class FElectraSubtitleDecoderFactoryTX3G : public IElectraSubtitleDecoderFactory
+{
+public:
+	virtual TSharedPtr<IElectraSubtitleDecoder, ESPMode::ThreadSafe> CreateDecoder(const FString& SubtitleCodecName) override;
+};
+
+
+void FElectraSubtitleDecoderTX3G::RegisterCodecs(IElectraSubtitleDecoderFactoryRegistry& InRegistry)
+{
+	static FElectraSubtitleDecoderFactoryTX3G Factory;
+	TArray<IElectraSubtitleDecoderFactoryRegistry::FCodecInfo> CodecInfos { { FString(TEXT("tx3g")), 0} };
+	InRegistry.AddDecoderFactory(CodecInfos, &Factory);
+}
+
+
+TSharedPtr<IElectraSubtitleDecoder, ESPMode::ThreadSafe> FElectraSubtitleDecoderFactoryTX3G::CreateDecoder(const FString& SubtitleCodecName)
+{
+	return MakeShared<FElectraSubtitleDecoderTX3G, ESPMode::ThreadSafe>();
+}
+
+
+
+
+
+
+class FSubtitleDecoderOutputTX3G : public ISubtitleDecoderOutput
+{
+public:
+	virtual ~FSubtitleDecoderOutputTX3G() = default;
+
+	void SetText(const FString& InText)
+	{
+		FTCHARToUTF8 Converted(*InText); // Convert to UTF8
+		TextAsArray.Empty();
+		TextAsArray.Append(reinterpret_cast<const uint8*>(Converted.Get()), Converted.Length());
+	}
+	
+	void SetDuration(const Electra::FTimeValue& InDuration)
+	{
+		Duration = InDuration.GetAsTimespan();
+	}
+	
+	void SetTimestamp(const Electra::FTimeValue& InTimestamp)
+	{
+		Timestamp.Time = InTimestamp.GetAsTimespan();
+		Timestamp.SequenceIndex = 0;
+	}
+	
+	void SetID(const FString& InID)
+	{
+		ID = InID;
+	}
+
+
+	virtual const TArray<uint8>& GetData() override
+	{
+		return TextAsArray;
+	}
+	
+	virtual FDecoderTimeStamp GetTime() const override
+	{
+		return Timestamp;
+	}
+
+	virtual FTimespan GetDuration() const override
+	{
+		return Duration;
+	}
+
+	virtual const FString& GetFormat() const override
+	{
+		static FString Format(TEXT("tx3g"));
+		return Format;
+	}
+	virtual const FString& GetID() const override
+	{
+		return ID;
+	}
+private:
+	TArray<uint8> TextAsArray;
+	FString ID;
+	FDecoderTimeStamp Timestamp;
+	FTimespan Duration;
+};
+
+
+
+
+FElectraSubtitleDecoderTX3G::FElectraSubtitleDecoderTX3G()
+{
+}
+
+FElectraSubtitleDecoderTX3G::~FElectraSubtitleDecoderTX3G()
+{
+}
+
+bool FElectraSubtitleDecoderTX3G::InitializeStreamWithCSD(const TArray<uint8>& InCSD, const Electra::FParamDict& InAdditionalInfo)
+{
+	#define RETURN_IF_ERROR(expr)								\
+		if (!expr)												\
+		{														\
+			UE_LOG(LogElectraSubtitles, Error, TEXT("Not enough CSD bytes to parse TX3G text sample entry"));	\
+			return false;										\
+		}														\
+
+
+	// Get dimension, placement offset and other data from the sideband dictionary.
+	Width = (int32) InAdditionalInfo.GetValue(TEXT("width")).SafeGetInt64(0);
+	Height = (int32) InAdditionalInfo.GetValue(TEXT("height")).SafeGetInt64(0);
+	TranslationX = (int32) InAdditionalInfo.GetValue(TEXT("offset_x")).SafeGetInt64(0);
+	TranslationY = (int32) InAdditionalInfo.GetValue(TEXT("offset_y")).SafeGetInt64(0);
+	Timescale = (uint32) InAdditionalInfo.GetValue(TEXT("timescale")).SafeGetInt64(0);
+
+	ElectraSubtitleDecoderTX3GUtils::FDataReaderMP4 r(InCSD);
+
+	// The CSD is a PlainTextSampleEntry of an mp4 (ISO/IEC 14496-12) file and interpreted as per
+	// ETSI TS 126 245 V11.0.0 - 5.16 Sample Description Format
+
+	RETURN_IF_ERROR(r.ReadBytes(nullptr, 6));
+	RETURN_IF_ERROR(r.Read(DataReferenceIndex));
+
+	RETURN_IF_ERROR(r.Read(DisplayFlags));
+	RETURN_IF_ERROR(r.Read(HorizontalJustification));
+	RETURN_IF_ERROR(r.Read(VerticalJustification));
+	RETURN_IF_ERROR(r.ReadBytes(BackgroundColorRGBA, 4));
+
+	// BoxRecord
+	RETURN_IF_ERROR(r.Read(BoxRecord.Top));
+	RETURN_IF_ERROR(r.Read(BoxRecord.Left));
+	RETURN_IF_ERROR(r.Read(BoxRecord.Bottom));
+	RETURN_IF_ERROR(r.Read(BoxRecord.Right));
+
+	// Style record
+	RETURN_IF_ERROR(r.Read(StyleRecord.StartChar));
+	RETURN_IF_ERROR(r.Read(StyleRecord.EndChar));
+	RETURN_IF_ERROR(r.Read(StyleRecord.FontID));
+	RETURN_IF_ERROR(r.Read(StyleRecord.FaceStyleFlags));
+	RETURN_IF_ERROR(r.Read(StyleRecord.FontSize));
+	RETURN_IF_ERROR(r.ReadBytes(StyleRecord.TextColorRGBA, 4));
+
+	// The font table is stored in a box. We skip over the box itself and parse it in place.
+	// Any boxes besides a single ftab box we ignore.
+	uint8 TempFTAB[8];
+	RETURN_IF_ERROR(r.ReadBytes(TempFTAB, 8));
+	if (TempFTAB[4] != 'f' || TempFTAB[5] != 't' || TempFTAB[6] != 'a' || TempFTAB[7] != 'b')
+	{
+		UE_LOG(LogElectraSubtitles, Error, TEXT("TX3G CSD is missing the ftab box!"));
+		return false;
+	}
+
+	// FontTableBox
+	uint16 EntryCount;
+	RETURN_IF_ERROR(r.Read(EntryCount));
+	FontRecords.Reserve(EntryCount);
+	for(int32 nEntry=0; nEntry<EntryCount; ++nEntry)
+	{
+		FFontRecord& fr = FontRecords.AddDefaulted_GetRef();
+		RETURN_IF_ERROR(r.Read(fr.FontID));
+		uint8 FontNameLength;
+		RETURN_IF_ERROR(r.Read(FontNameLength));
+		if (FontNameLength)
+		{
+			if (!r.ReadString(fr.FontName, FontNameLength))
+			{
+				return false;
+			}
+		}
+	}
+
+	#undef RETURN_IF_ERROR
+	return true;
+}
+
+IElectraSubtitleDecoder::FOnSubtitleReceivedDelegate& FElectraSubtitleDecoderTX3G::GetParsedSubtitleReceiveDelegate()
+{
+	return ParsedSubtitleDelegate;
+}
+
+Electra::FTimeValue FElectraSubtitleDecoderTX3G::GetStreamedDeliveryTimeOffset()
+{
+	return Electra::FTimeValue::GetZero();
+}
+
+void FElectraSubtitleDecoderTX3G::AddStreamedSubtitleData(const TArray<uint8>& InData, Electra::FTimeValue InAbsoluteTimestamp, Electra::FTimeValue InDuration, const Electra::FParamDict& InAdditionalInfo)
+{
+	#define RETURN_IF_ERROR(expr)															\
+		if (!expr)																			\
+		{																					\
+			UE_LOG(LogElectraSubtitles, Error, TEXT("Bad TX3G text sample, ignoring."));	\
+			return;																			\
+		}																					\
+	
+	ElectraSubtitleDecoderTX3GUtils::FDataReaderMP4 r(InData);
+
+	// Start a subtitle entry.
+	FSubtitleEntry Sub;
+	// By default it uses the same box area as defined in the CSD. Each subtitle can override this locally.
+	Sub.Box = BoxRecord;
+
+	uint16 TextLen;
+	RETURN_IF_ERROR(r.Read(TextLen));
+	if (TextLen)
+	{
+		if (r.ReadString(Sub.Text, TextLen))
+		{
+			// Get the number of trailing bytes that describe styling boxes.
+			while(r.GetNumBytesRemaining() > 0)
+			{
+				uint32 BoxLen, BoxType;
+				if (!r.Read(BoxLen) || !r.Read(BoxType))
+				{
+					UE_LOG(LogElectraSubtitles, Error, TEXT("Bad TX3G text sample formatting box, ignoring."));
+					break;
+				}
+				if (BoxLen < 8 || (int32)(BoxLen-8 )> r.GetNumBytesRemaining())
+				{
+					UE_LOG(LogElectraSubtitles, Error, TEXT("Bad TX3G text sample formatting box, ignoring."));
+					break;
+				}
+				if (BoxType == ElectraSubtitleDecoderTX3GUtils::FDataReaderMP4::BoxType_styl)
+				{
+					uint16 NumStyleRecords;
+					RETURN_IF_ERROR(r.Read(NumStyleRecords));
+					for(int32 i=0; i<NumStyleRecords; ++i)
+					{
+						FStyleRecord sr;
+						RETURN_IF_ERROR(r.Read(sr.StartChar));
+						RETURN_IF_ERROR(r.Read(sr.EndChar));
+						RETURN_IF_ERROR(r.Read(sr.FontID));
+						RETURN_IF_ERROR(r.Read(sr.FaceStyleFlags));
+						RETURN_IF_ERROR(r.Read(sr.FontSize));
+						RETURN_IF_ERROR(r.ReadBytes(sr.TextColorRGBA, 4));
+						Sub.Styles.Emplace(MoveTemp(sr));
+					}
+				}
+				else if (BoxType == ElectraSubtitleDecoderTX3GUtils::FDataReaderMP4::BoxType_hlit)
+				{
+					FHighlight hl;
+					RETURN_IF_ERROR(r.Read(hl.StartCharOffset));
+					RETURN_IF_ERROR(r.Read(hl.EndCharOffset));
+					Sub.Highlights.Emplace(MoveTemp(hl));
+				}
+				else if (BoxType == ElectraSubtitleDecoderTX3GUtils::FDataReaderMP4::BoxType_hclr)
+				{
+					FHighlightColor hc;
+					RETURN_IF_ERROR(r.ReadBytes(hc.TextColorRGBA, 4));
+					Sub.HighlightColor.Emplace(MoveTemp(hc));
+				}
+				else if (BoxType == ElectraSubtitleDecoderTX3GUtils::FDataReaderMP4::BoxType_krok)
+				{
+					FTextKaraoke kr;
+					RETURN_IF_ERROR(r.Read(kr.HighlightStartTime));
+					uint16 NumEntries;
+					RETURN_IF_ERROR(r.Read(NumEntries));
+					kr.Entries.Reserve(NumEntries);
+					for(int32 i=0; i<NumEntries; ++i)
+					{
+						FTextKaraoke::FEntry& e=kr.Entries.AddDefaulted_GetRef();
+						RETURN_IF_ERROR(r.Read(e.HightlightEndTime));
+						RETURN_IF_ERROR(r.Read(e.StartCharOffset));
+						RETURN_IF_ERROR(r.Read(e.EndCharOffset));
+					}
+					Sub.Karaoke.Emplace(MoveTemp(kr));
+				}
+				else if (BoxType == ElectraSubtitleDecoderTX3GUtils::FDataReaderMP4::BoxType_dlay)
+				{
+					RETURN_IF_ERROR(r.Read(Sub.ScrollDelay.ScrollDelay));
+				}
+				else if (BoxType == ElectraSubtitleDecoderTX3GUtils::FDataReaderMP4::BoxType_href)
+				{
+					uint8 Len;
+					FHyperText ht;
+					RETURN_IF_ERROR(r.Read(ht.StartCharOffset));
+					RETURN_IF_ERROR(r.Read(ht.EndCharOffset));
+					RETURN_IF_ERROR(r.Read(Len));
+					RETURN_IF_ERROR(r.ReadString(ht.URL, Len));
+					RETURN_IF_ERROR(r.Read(Len));
+					RETURN_IF_ERROR(r.ReadString(ht.AltString, Len));
+					Sub.HyperTexts.Emplace(MoveTemp(ht));
+				}
+				else if (BoxType == ElectraSubtitleDecoderTX3GUtils::FDataReaderMP4::BoxType_tbox)
+				{
+					RETURN_IF_ERROR(r.Read(Sub.Box.Top));
+					RETURN_IF_ERROR(r.Read(Sub.Box.Left));
+					RETURN_IF_ERROR(r.Read(Sub.Box.Bottom));
+					RETURN_IF_ERROR(r.Read(Sub.Box.Right));
+				}
+				else if (BoxType == ElectraSubtitleDecoderTX3GUtils::FDataReaderMP4::BoxType_blnk)
+				{
+					FBlink bk;
+					RETURN_IF_ERROR(r.Read(bk.StartCharOffset));
+					RETURN_IF_ERROR(r.Read(bk.EndCharOffset));
+					Sub.Blinks.Emplace(MoveTemp(bk));
+				}
+				else if (BoxType == ElectraSubtitleDecoderTX3GUtils::FDataReaderMP4::BoxType_twrp)
+				{
+					FTextWrap wr;
+					RETURN_IF_ERROR(r.Read(wr.WrapFlag));
+					Sub.TextWrap.Emplace(MoveTemp(wr));
+				}
+				else
+				{
+					if (!r.ReadBytes(nullptr, BoxLen - 8))
+					{
+						break;
+					}
+				}
+			}
+		}
+	}
+	else
+	{
+		// An empty entry.
+	}
+
+	TSharedPtr<FSubtitleDecoderOutputTX3G, ESPMode::ThreadSafe> Out = MakeShared<FSubtitleDecoderOutputTX3G, ESPMode::ThreadSafe>();
+	Out->SetTimestamp(InAbsoluteTimestamp);
+	Out->SetDuration(InDuration);
+	Out->SetText(Sub.Text);
+	Out->SetID(LexToString(++NextID));
+	ParsedSubtitleDelegate.Broadcast(Out);
+
+	#undef RETURN_IF_ERROR
+}
+
+void FElectraSubtitleDecoderTX3G::SignalStreamedSubtitleEOD()
+{
+}
+
+void FElectraSubtitleDecoderTX3G::Flush()
+{
+}
+
+void FElectraSubtitleDecoderTX3G::Start()
+{
+}
+
+void FElectraSubtitleDecoderTX3G::Stop()
+{
+}
+
+void FElectraSubtitleDecoderTX3G::UpdatePlaybackPosition(Electra::FTimeValue InAbsolutePosition, Electra::FTimeValue InLocalPosition)
+{
+}
+
diff --git a/Engine/Plugins/Media/ElectraSubtitles/Source/ElectraSubtitles/Private/tx3g/ElectraSubtitleDecoder_TX3G.h b/Engine/Plugins/Media/ElectraSubtitles/Source/ElectraSubtitles/Private/tx3g/ElectraSubtitleDecoder_TX3G.h
new file mode 100644
--- /dev/null
+++ b/Engine/Plugins/Media/ElectraSubtitles/Source/ElectraSubtitles/Private/tx3g/ElectraSubtitleDecoder_TX3G.h
@@ -0,0 +1,136 @@
+// Copyright Epic Games, Inc. All Rights Reserved.
+
+#pragma once
+
+#include "ElectraSubtitleDecoder.h"
+#include "ElectraSubtitleDecoderFactory.h"
+
+/**
+ * 3GPP / TX3G subtitle decoder (ETSI TS 126 245 V11.0.0)
+ */
+class FElectraSubtitleDecoderTX3G : public IElectraSubtitleDecoder
+{
+public:
+	static void RegisterCodecs(IElectraSubtitleDecoderFactoryRegistry& InRegistry);
+
+	FElectraSubtitleDecoderTX3G();
+
+	virtual ~FElectraSubtitleDecoderTX3G();
+
+	//-------------------------------------------------------------------------
+	// Methods from IElectraSubtitleDecoder
+	//
+	virtual bool InitializeStreamWithCSD(const TArray<uint8>& InCSD, const Electra::FParamDict& InAdditionalInfo) override;
+	virtual IElectraSubtitleDecoder::FOnSubtitleReceivedDelegate& GetParsedSubtitleReceiveDelegate() override;
+	virtual Electra::FTimeValue GetStreamedDeliveryTimeOffset() override;
+	virtual void AddStreamedSubtitleData(const TArray<uint8>& InData, Electra::FTimeValue InAbsoluteTimestamp, Electra::FTimeValue InDuration, const Electra::FParamDict& InAdditionalInfo) override;
+	virtual void SignalStreamedSubtitleEOD() override;
+	virtual void Flush() override;
+	virtual void Start() override;
+	virtual void Stop() override;
+	virtual void UpdatePlaybackPosition(Electra::FTimeValue InAbsolutePosition, Electra::FTimeValue InLocalPosition) override;
+
+private:
+	struct FBoxRecord
+	{
+		int16 Top = 0;
+		int16 Left = 0;
+		int16 Bottom = 0;
+		int16 Right = 0;
+	};
+
+	struct FStyleRecord
+	{
+		uint16 StartChar = 0;
+		uint16 EndChar = 0;
+		uint16 FontID = 0;
+		uint8 FaceStyleFlags = 0;
+		uint8 FontSize = 0;
+		uint8 TextColorRGBA[4] { 0 };
+	};
+
+	struct FFontRecord
+	{
+		uint16 FontID = 0;
+		FString FontName;
+	};
+
+	struct FHighlight
+	{
+		uint16 StartCharOffset = 0;
+		uint16 EndCharOffset = 0;
+	};
+
+	struct FHighlightColor
+	{
+		uint8 TextColorRGBA[4] { 0 };
+	};
+
+	struct FTextKaraoke
+	{
+		uint32 HighlightStartTime = 0;
+		struct FEntry
+		{
+			uint32 HightlightEndTime = 0;
+			uint16 StartCharOffset = 0;
+			uint16 EndCharOffset = 0;
+		};
+		TArray<FEntry> Entries;
+	};
+
+	struct FScrollDelay
+	{
+		uint32 ScrollDelay = 0;
+	};
+
+	struct FHyperText
+	{
+		uint16 StartCharOffset = 0;
+		uint16 EndCharOffset = 0;
+		FString URL;
+		FString AltString;
+	};
+
+	struct FBlink
+	{
+		uint16 StartCharOffset = 0;
+		uint16 EndCharOffset = 0;
+	};
+
+	struct FTextWrap
+	{
+		uint8 WrapFlag = 0;
+	};
+
+
+	struct FSubtitleEntry
+	{
+		FString						Text;
+		FBoxRecord					Box;
+		TArray<FStyleRecord>		Styles;
+		TArray<FHighlight>			Highlights;
+		TOptional<FHighlightColor>	HighlightColor;
+		TOptional<FTextKaraoke>		Karaoke;
+		FScrollDelay				ScrollDelay;
+		TArray<FHyperText>			HyperTexts;
+		TArray<FBlink>				Blinks;
+		TOptional<FTextWrap>		TextWrap;
+	};
+
+	FOnSubtitleReceivedDelegate ParsedSubtitleDelegate;
+	uint32 NextID = 0;
+
+	uint16				DataReferenceIndex = 0;
+	int32				Width = 0;
+	int32				Height = 0;
+	int32				TranslationX = 0;
+	int32				TranslationY = 0;
+	uint32				Timescale = 0;
+	uint32				DisplayFlags = 0;
+	int8				HorizontalJustification = 0;
+	int8				VerticalJustification = 0;
+	uint8				BackgroundColorRGBA[4] {0};
+	FBoxRecord			BoxRecord;
+	FStyleRecord		StyleRecord;
+	TArray<FFontRecord>	FontRecords;
+};
diff --git a/Engine/Plugins/Media/ElectraSubtitles/Source/ElectraSubtitles/Public/IElectraSubtitleDecoder.h b/Engine/Plugins/Media/ElectraSubtitles/Source/ElectraSubtitles/Public/IElectraSubtitleDecoder.h
new file mode 100644
--- /dev/null
+++ b/Engine/Plugins/Media/ElectraSubtitles/Source/ElectraSubtitles/Public/IElectraSubtitleDecoder.h
@@ -0,0 +1,52 @@
+// Copyright Epic Games, Inc. All Rights Reserved.
+
+#pragma once
+
+#include "CoreMinimal.h"
+#include "PlayerTime.h"
+#include "ParameterDictionary.h"
+#include "MediaSubtitleDecoderOutput.h"
+
+
+/**
+ * Base class of a subtitle decoder.
+ */
+class IElectraSubtitleDecoder
+{
+public:
+	virtual ~IElectraSubtitleDecoder() = default;
+
+	//-------------------------------------------------------------------------
+	// Initialization methods
+	//
+	//! Initializes with codec specific data and additional attributes.
+	virtual bool InitializeStreamWithCSD(const TArray<uint8>& InCSD, const Electra::FParamDict& InAdditionalInfo) = 0;
+
+	//-------------------------------------------------------------------------
+	// Delegate to which the application can subscribe to receive parsed subtitles.
+	//
+	DECLARE_MULTICAST_DELEGATE_OneParam(FOnSubtitleReceivedDelegate, ISubtitleDecoderOutputPtr /* DecodedSubtitle */);
+	virtual FOnSubtitleReceivedDelegate& GetParsedSubtitleReceiveDelegate() = 0;
+
+	//-------------------------------------------------------------------------
+	// Data handling methods
+	//
+	//! Returns a time offset at which streamed subtitles should be delivered ahead of their actual timestamp (if they are available earlier)
+	virtual Electra::FTimeValue GetStreamedDeliveryTimeOffset() = 0;
+	//! Adds streamed subtitle data.
+	virtual void AddStreamedSubtitleData(const TArray<uint8>& InData, Electra::FTimeValue InAbsoluteTimestamp, Electra::FTimeValue InDuration, const Electra::FParamDict& InAdditionalInfo) = 0;
+	//! Signals no more streamed subtitles will be arriving.
+	virtual void SignalStreamedSubtitleEOD() = 0;
+	//! Flushes all input and output.
+	virtual void Flush() = 0;
+
+
+	//-------------------------------------------------------------------------
+	// Handling methods
+	//
+	virtual void Start() = 0;
+	virtual void Stop() = 0;
+	virtual void UpdatePlaybackPosition(Electra::FTimeValue InAbsolutePosition, Electra::FTimeValue InLocalPosition) = 0;
+
+};
+
diff --git a/Engine/Plugins/Media/ElectraSubtitles/Source/ElectraSubtitles/Public/IElectraSubtitleModule.h b/Engine/Plugins/Media/ElectraSubtitles/Source/ElectraSubtitles/Public/IElectraSubtitleModule.h
new file mode 100644
--- /dev/null
+++ b/Engine/Plugins/Media/ElectraSubtitles/Source/ElectraSubtitles/Public/IElectraSubtitleModule.h
@@ -0,0 +1,86 @@
+// Copyright Epic Games, Inc. All Rights Reserved.
+
+#pragma once
+
+#include "CoreMinimal.h"
+#include "Templates/SharedPointer.h"
+#include "Features/IModularFeatures.h"
+#include "Features/IModularFeature.h"
+#include "Modules/ModuleInterface.h"
+#include "Modules/ModuleManager.h"
+
+#include "IElectraSubtitleDecoder.h"
+
+
+class IElectraSubtitleModularFeature : public IModularFeature
+{
+public:
+	virtual ~IElectraSubtitleModularFeature() = default;
+
+	/**
+	 * Queries whether or not this plugin's modular feature is capable of handling subtitles of
+	 * the specified format. The codec string is a commonly known format like "wvtt", "tx3g", "imsc1", etc.
+	 * but may also be a MIME type like "application/ttml+xml" for sideloaded subtitles.
+	 */
+	virtual bool SupportsFormat(const FString& SubtitleCodecName) const = 0;
+
+	/**
+	 * Returns a list of the supported subtitle codecs.
+	 */
+	virtual void GetSupportedFormats(TArray<FString>& OutSupportedCodecNames) const = 0;
+
+	/**
+	 * Returns a priority for this plugin's subtitle format feature.
+	 * The return value is not standardized, but the decoder returning the largest value for a format
+	 * will be used to handle the subtitles.
+	 * This is mainly used if you provide a plugin of your own that handles a format better than another
+	 * plugin does and you want to make sure yours will be used.
+	 */
+	virtual int32 GetPriorityForFormat(const FString& SubtitleCodecName) const = 0;
+
+	/**
+	 * Called to create a decoder for the given format from this plugin's modular feature.
+	 */
+	virtual TSharedPtr<IElectraSubtitleDecoder, ESPMode::ThreadSafe> CreateDecoderForFormat(const FString& SubtitleCodecName) = 0;
+};
+
+
+
+/**
+ * Interface for the `ElectraSubtitles` module.
+ */
+class IElectraSubtitlesModule : public IModuleInterface
+{
+public:
+	virtual ~IElectraSubtitlesModule() = default;
+
+	static FName GetModularFeatureName()
+	{
+		static FName FeatureName = FName(TEXT("ElectraSubtitles"));
+		return FeatureName;
+	}
+
+
+#if 0
+	/**
+	 * Singleton-like access to IElectraSubtitlesModule
+	 *
+	 * @return Returns IElectraSubtitlesModule singleton instance, loading the module on demand if needed
+	 */
+	static inline IElectraSubtitlesModule& Get()
+	{
+		return FModuleManager::LoadModuleChecked<IElectraSubtitlesModule>("ElectraSubtitles");
+	}
+
+	/**
+	 * Checks to see if this module is loaded and ready.  It is only valid to call Get() if IsAvailable() returns true.
+	 *
+	 * @return True if the module is loaded and ready to use
+	 */
+	static inline bool IsAvailable()
+	{
+		return FModuleManager::Get().IsModuleLoaded("ElectraSubtitles");
+	}
+#endif
+};
+
diff --git a/Engine/Plugins/Media/ElectraUtil/Source/ElectraBase/Public/MediaSubtitleDecoderOutput.h b/Engine/Plugins/Media/ElectraUtil/Source/ElectraBase/Public/MediaSubtitleDecoderOutput.h
new file mode 100644
--- /dev/null
+++ b/Engine/Plugins/Media/ElectraUtil/Source/ElectraBase/Public/MediaSubtitleDecoderOutput.h
@@ -0,0 +1,23 @@
+// Copyright Epic Games, Inc. All Rights Reserved.
+#pragma once
+
+#include <CoreMinimal.h>
+#include <Misc/Timespan.h>
+#include <Templates/SharedPointer.h>
+
+#include "MediaDecoderOutput.h"
+
+class ISubtitleDecoderOutput
+{
+public:
+	virtual ~ISubtitleDecoderOutput() = default;
+
+	virtual const TArray<uint8>& GetData() = 0;
+	virtual FDecoderTimeStamp GetTime() const = 0;
+	virtual FTimespan GetDuration() const = 0;
+
+	virtual const FString& GetFormat() const = 0;
+	virtual const FString& GetID() const = 0;
+};
+
+using ISubtitleDecoderOutputPtr = TSharedPtr<ISubtitleDecoderOutput, ESPMode::ThreadSafe>;
