ElectraSubtitles: added missing #include causing compilation error in CI

diff --git a/Engine/Plugins/Media/ElectraSubtitles/Source/ElectraSubtitles/Private/ElectraSubtitleDecoder.h b/Engine/Plugins/Media/ElectraSubtitles/Source/ElectraSubtitles/Private/ElectraSubtitleDecoder.h
--- a/Engine/Plugins/Media/ElectraSubtitles/Source/ElectraSubtitles/Private/ElectraSubtitleDecoder.h
+++ b/Engine/Plugins/Media/ElectraSubtitles/Source/ElectraSubtitles/Private/ElectraSubtitleDecoder.h
@@ -3,6 +3,4 @@
 #pragma once
 
 #include "IElectraSubtitleDecoder.h"
-
-
-
+#include "ElectraSubtitleModule.h"
