ElectraSubtitles: Added WebVTT (wvtt) plaintext subtitle parser

#rnx

diff --git a/Engine/Plugins/Media/ElectraSubtitles/Source/ElectraSubtitles/Private/ElectraSubtitleModule.cpp b/Engine/Plugins/Media/ElectraSubtitles/Source/ElectraSubtitles/Private/ElectraSubtitleModule.cpp
--- a/Engine/Plugins/Media/ElectraSubtitles/Source/ElectraSubtitles/Private/ElectraSubtitleModule.cpp
+++ b/Engine/Plugins/Media/ElectraSubtitles/Source/ElectraSubtitles/Private/ElectraSubtitleModule.cpp
@@ -8,6 +8,7 @@
 #include "ElectraSubtitleDecoderFactory.h"
 
 #include "tx3g/ElectraSubtitleDecoder_TX3G.h"
+#include "wvtt/ElectraSubtitleDecoder_WVTT.h"
 
 #define LOCTEXT_NAMESPACE "ElectraSubtitlesModule"
 
@@ -56,6 +57,7 @@ public:
 	{
 		// Register the codecs this plugin provides.
 		FElectraSubtitleDecoderTX3G::RegisterCodecs(ModularFeature);
+		FElectraSubtitleDecoderWVTT::RegisterCodecs(ModularFeature);
 
 		IModularFeatures::Get().RegisterModularFeature(GetModularFeatureName(), &ModularFeature);
 	}
diff --git a/Engine/Plugins/Media/ElectraSubtitles/Source/ElectraSubtitles/Private/wvtt/ElectraSubtitleDecoder_WVTT.cpp b/Engine/Plugins/Media/ElectraSubtitles/Source/ElectraSubtitles/Private/wvtt/ElectraSubtitleDecoder_WVTT.cpp
new file mode 100644
--- /dev/null
+++ b/Engine/Plugins/Media/ElectraSubtitles/Source/ElectraSubtitles/Private/wvtt/ElectraSubtitleDecoder_WVTT.cpp
@@ -0,0 +1,512 @@
+// Copyright Epic Games, Inc. All Rights Reserved.
+
+#include "wvtt/ElectraSubtitleDecoder_WVTT.h"
+#include "Containers/StringConv.h"
+
+namespace ElectraSubtitleDecoderWVTTUtils
+{
+
+#if !PLATFORM_LITTLE_ENDIAN
+	static inline uint8 GetFromBigEndian(uint8 value)		{ return value; }
+	static inline int8 GetFromBigEndian(int8 value)			{ return value; }
+	static inline uint16 GetFromBigEndian(uint16 value)		{ return value; }
+	static inline int16 GetFromBigEndian(int16 value)		{ return value; }
+	static inline int32 GetFromBigEndian(int32 value)		{ return value; }
+	static inline uint32 GetFromBigEndian(uint32 value)		{ return value; }
+	static inline int64 GetFromBigEndian(int64 value)		{ return value; }
+	static inline uint64 GetFromBigEndian(uint64 value)		{ return value; }
+#else
+	static inline uint16 EndianSwap(uint16 value)			{ return (value >> 8) | (value << 8); }
+	static inline int16 EndianSwap(int16 value)				{ return int16(EndianSwap(uint16(value))); }
+	static inline uint32 EndianSwap(uint32 value)			{ return (value << 24) | ((value & 0xff00) << 8) | ((value >> 8) & 0xff00) | (value >> 24); }
+	static inline int32 EndianSwap(int32 value)				{ return int32(EndianSwap(uint32(value))); }
+	static inline uint64 EndianSwap(uint64 value)			{ return (uint64(EndianSwap(uint32(value & 0xffffffffU))) << 32) | uint64(EndianSwap(uint32(value >> 32))); }
+	static inline int64 EndianSwap(int64 value)				{ return int64(EndianSwap(uint64(value)));}
+	static inline uint8 GetFromBigEndian(uint8 value)		{ return value; }
+	static inline int8 GetFromBigEndian(int8 value)			{ return value; }
+	static inline uint16 GetFromBigEndian(uint16 value)		{ return EndianSwap(value); }
+	static inline int16 GetFromBigEndian(int16 value)		{ return EndianSwap(value); }
+	static inline int32 GetFromBigEndian(int32 value)		{ return EndianSwap(value); }
+	static inline uint32 GetFromBigEndian(uint32 value)		{ return EndianSwap(value); }
+	static inline int64 GetFromBigEndian(int64 value)		{ return EndianSwap(value); }
+	static inline uint64 GetFromBigEndian(uint64 value)		{ return EndianSwap(value); }
+#endif
+
+	template <typename T>
+	T ValueFromBigEndian(const T value)
+	{
+		return GetFromBigEndian(value);
+	}
+
+
+	class FDataReaderMP4
+	{
+	public:
+		FDataReaderMP4(const TArray<uint8>& InDataBufferToReadFrom) : DataBufferRef(InDataBufferToReadFrom)
+		{}
+
+		int32 GetCurrentOffset() const
+		{ return CurrentOffset;	}
+
+		int32 GetNumBytesRemaining() const
+		{ return DataBufferRef.Num() - GetCurrentOffset(); }
+
+		template <typename T>
+		bool Read(T& value)
+		{
+			T Temp = 0;
+			int64 NumRead = ReadData(&Temp, sizeof(T));
+			if (NumRead == sizeof(T))
+			{
+				value = ElectraSubtitleDecoderWVTTUtils::ValueFromBigEndian(Temp);
+				return true;
+			}
+			return false;
+		}
+
+		bool ReadString(FString& OutString, uint16 NumBytes)
+		{
+			OutString.Empty();
+			if (NumBytes == 0)
+			{
+				return true;
+			}
+			TArray<uint8> Buf;
+			Buf.AddUninitialized(NumBytes);
+			if (ReadBytes(Buf.GetData(), NumBytes))
+			{
+				// Check for UTF16 BOM
+				if (NumBytes >= 2 && ((Buf[0] == 0xff && Buf[1] == 0xfe) || (Buf[0] == 0xfe && Buf[1] == 0xff)))
+				{
+					UE_LOG(LogElectraSubtitles, Error, TEXT("WVTT uses UTF16 which is not supported"));
+					return false;
+				}
+				FUTF8ToTCHAR cnv((const ANSICHAR*)Buf.GetData(), NumBytes);
+				OutString = FString(cnv.Length(), cnv.Get());
+				return true;
+			}
+			return false;
+		}
+
+		bool ReadBytes(void* Buffer, int32 NumBytes)
+		{
+			int32 NumRead = ReadData(Buffer, NumBytes);
+			return NumRead == NumBytes;
+		}
+
+#define MAKE_BOX_ATOM(a,b,c,d) (uint32)((uint32)a << 24) | ((uint32)b << 16) | ((uint32)c << 8) | ((uint32)d)
+		static const uint32 BoxType_vttC = MAKE_BOX_ATOM('v', 't', 't', 'C');
+		static const uint32 BoxType_vlab = MAKE_BOX_ATOM('v', 'l', 'a', 'b');
+		static const uint32 BoxType_vtte = MAKE_BOX_ATOM('v', 't', 't', 'e');
+		static const uint32 BoxType_vtta = MAKE_BOX_ATOM('v', 't', 't', 'a');
+		static const uint32 BoxType_vttc = MAKE_BOX_ATOM('v', 't', 't', 'c');
+		static const uint32 BoxType_vsid = MAKE_BOX_ATOM('v', 's', 'i', 'd');
+		static const uint32 BoxType_ctim = MAKE_BOX_ATOM('c', 't', 'i', 'm');
+		static const uint32 BoxType_iden = MAKE_BOX_ATOM('i', 'd', 'e', 'n');
+		static const uint32 BoxType_sttg = MAKE_BOX_ATOM('s', 't', 't', 'g');
+		static const uint32 BoxType_payl = MAKE_BOX_ATOM('p', 'a', 'y', 'l');
+#undef MAKE_BOX_ATOM
+	private:
+		int32 ReadData(void* IntoBuffer, int32 NumBytesToRead)
+		{
+			if (NumBytesToRead == 0)
+			{
+				return 0;
+			}
+			int32 NumAvail = DataBufferRef.Num() - CurrentOffset;
+			if (NumAvail >= NumBytesToRead)
+			{
+				if (IntoBuffer)
+				{
+					FMemory::Memcpy(IntoBuffer, DataBufferRef.GetData() + CurrentOffset, NumBytesToRead);
+				}
+				CurrentOffset += NumBytesToRead;
+				return NumBytesToRead;
+			}
+			return -1;
+		}
+
+		const TArray<uint8>& DataBufferRef;
+		int32 CurrentOffset = 0;
+	};
+
+}
+
+
+
+class FElectraSubtitleDecoderFactoryWVTT : public IElectraSubtitleDecoderFactory
+{
+public:
+	virtual TSharedPtr<IElectraSubtitleDecoder, ESPMode::ThreadSafe> CreateDecoder(const FString& SubtitleCodecName) override;
+};
+
+
+void FElectraSubtitleDecoderWVTT::RegisterCodecs(IElectraSubtitleDecoderFactoryRegistry& InRegistry)
+{
+	static FElectraSubtitleDecoderFactoryWVTT Factory;
+	TArray<IElectraSubtitleDecoderFactoryRegistry::FCodecInfo> CodecInfos { { FString(TEXT("wvtt")), 0} };
+	InRegistry.AddDecoderFactory(CodecInfos, &Factory);
+}
+
+
+TSharedPtr<IElectraSubtitleDecoder, ESPMode::ThreadSafe> FElectraSubtitleDecoderFactoryWVTT::CreateDecoder(const FString& SubtitleCodecName)
+{
+	return MakeShared<FElectraSubtitleDecoderWVTT, ESPMode::ThreadSafe>();
+}
+
+
+
+
+
+
+class FSubtitleDecoderOutputWVTT : public ISubtitleDecoderOutput
+{
+public:
+	virtual ~FSubtitleDecoderOutputWVTT() = default;
+
+	void SetText(const FString& InText)
+	{
+		FTCHARToUTF8 Converted(*InText); // Convert to UTF8
+		TextAsArray.Empty();
+		TextAsArray.Append(reinterpret_cast<const uint8*>(Converted.Get()), Converted.Length());
+	}
+	
+	void SetDuration(const Electra::FTimeValue& InDuration)
+	{
+		Duration = InDuration.GetAsTimespan();
+	}
+	
+	void SetTimestamp(const Electra::FTimeValue& InTimestamp)
+	{
+		Timestamp.Time = InTimestamp.GetAsTimespan();
+		Timestamp.SequenceIndex = 0;
+	}
+	
+	void SetID(const FString& InID)
+	{
+		ID = InID;
+	}
+
+
+	virtual const TArray<uint8>& GetData() override
+	{
+		return TextAsArray;
+	}
+	
+	virtual FDecoderTimeStamp GetTime() const override
+	{
+		return Timestamp;
+	}
+
+	virtual FTimespan GetDuration() const override
+	{
+		return Duration;
+	}
+
+	virtual const FString& GetFormat() const override
+	{
+		static FString Format(TEXT("wvtt"));
+		return Format;
+	}
+	virtual const FString& GetID() const override
+	{
+		return ID;
+	}
+private:
+	TArray<uint8> TextAsArray;
+	FString ID;
+	FDecoderTimeStamp Timestamp;
+	FTimespan Duration;
+};
+
+
+
+
+FElectraSubtitleDecoderWVTT::FElectraSubtitleDecoderWVTT()
+{
+}
+
+FElectraSubtitleDecoderWVTT::~FElectraSubtitleDecoderWVTT()
+{
+}
+
+bool FElectraSubtitleDecoderWVTT::InitializeStreamWithCSD(const TArray<uint8>& InCSD, const Electra::FParamDict& InAdditionalInfo)
+{
+	#define RETURN_IF_ERROR(expr)								\
+		if (!expr)												\
+		{														\
+			UE_LOG(LogElectraSubtitles, Error, TEXT("Not enough CSD bytes to parse WVTT text sample entry"));	\
+			return false;										\
+		}														\
+
+	// Get dimension, placement offset and other data from the sideband dictionary.
+	Width = (int32) InAdditionalInfo.GetValue(TEXT("width")).SafeGetInt64(0);
+	Height = (int32) InAdditionalInfo.GetValue(TEXT("height")).SafeGetInt64(0);
+	TranslationX = (int32) InAdditionalInfo.GetValue(TEXT("offset_x")).SafeGetInt64(0);
+	TranslationY = (int32) InAdditionalInfo.GetValue(TEXT("offset_y")).SafeGetInt64(0);
+	Timescale = (uint32) InAdditionalInfo.GetValue(TEXT("timescale")).SafeGetInt64(0);
+
+	ElectraSubtitleDecoderWVTTUtils::FDataReaderMP4 r(InCSD);
+
+	// The CSD is a WVTTSampleEntry of an mp4 (ISO/IEC 14496-12) file and interpreted as per
+	// ISO/IEC 14496-30 Section 7.5 Sample entry format
+	RETURN_IF_ERROR(r.ReadBytes(nullptr, 6));
+	RETURN_IF_ERROR(r.Read(DataReferenceIndex));
+
+	// Read the configuration boxes.
+	// There needs to be a 'vttC' box.
+	// 'vlab' and 'btrt' boxes are optional. Anything else is ignored.
+	while(r.GetNumBytesRemaining() > 0)
+	{
+		uint32 BoxLen, BoxType;
+		if (!r.Read(BoxLen) || !r.Read(BoxType))
+		{
+			UE_LOG(LogElectraSubtitles, Error, TEXT("Bad WVTT box in CSD, ignoring."));
+			break;
+		}
+		if (BoxLen < 8 || (int32)(BoxLen-8 )> r.GetNumBytesRemaining())
+		{
+			UE_LOG(LogElectraSubtitles, Error, TEXT("Bad WVTT box in CSD, ignoring."));
+			break;
+		}
+		
+		if (BoxType == ElectraSubtitleDecoderWVTTUtils::FDataReaderMP4::BoxType_vttC)
+		{
+			RETURN_IF_ERROR(r.ReadString(Configuration, BoxLen - 8));
+		}
+		else if (BoxType == ElectraSubtitleDecoderWVTTUtils::FDataReaderMP4::BoxType_vlab)
+		{
+			RETURN_IF_ERROR(r.ReadString(Label, BoxLen - 8));
+		}
+		else
+		{
+			// Note: We ignore a potential 'btrt' box here. There is no use for us.
+			if (!r.ReadBytes(nullptr, BoxLen - 8))
+			{
+				break;
+			}
+		}
+	}
+
+	// The "configuration" is everything up to the first cue, so at least the string "WEBVTT".
+	if (Configuration.IsEmpty() || !Configuration.StartsWith(TEXT("WEBVTT")))
+	{
+		UE_LOG(LogElectraSubtitles, Error, TEXT("Bad WVTT configuration box 'vttC'!"));
+		return false;
+	}
+
+	#undef RETURN_IF_ERROR
+	return true;
+}
+
+IElectraSubtitleDecoder::FOnSubtitleReceivedDelegate& FElectraSubtitleDecoderWVTT::GetParsedSubtitleReceiveDelegate()
+{
+	return ParsedSubtitleDelegate;
+}
+
+Electra::FTimeValue FElectraSubtitleDecoderWVTT::GetStreamedDeliveryTimeOffset()
+{
+	return Electra::FTimeValue::GetZero();
+}
+
+void FElectraSubtitleDecoderWVTT::AddStreamedSubtitleData(const TArray<uint8>& InData, Electra::FTimeValue InAbsoluteTimestamp, Electra::FTimeValue InDuration, const Electra::FParamDict& InAdditionalInfo)
+{
+	#define RETURN_IF_ERROR(expr)																\
+		if (!expr)																				\
+		{																						\
+			UE_LOG(LogElectraSubtitles, Error, TEXT("Bad WVTT text sample box, ignoring."));	\
+			return;																				\
+		}																						\
+	
+	ElectraSubtitleDecoderWVTTUtils::FDataReaderMP4 r(InData);
+
+	// List of collected subtitles.
+	TArray<FSubtitleEntry> Subtitles;
+	// The currently worked-on subtitle.
+	FSubtitleEntry Sub;
+
+	bool bInsideCue = false;
+	while(r.GetNumBytesRemaining() > 0)
+	{
+		uint32 BoxLen, BoxType;
+		if (!r.Read(BoxLen) || !r.Read(BoxType))
+		{
+			UE_LOG(LogElectraSubtitles, Error, TEXT("Bad WVTT text sample box, ignoring."));
+			break;
+		}
+		if (BoxLen < 8 || (int32)(BoxLen-8 )> r.GetNumBytesRemaining())
+		{
+			UE_LOG(LogElectraSubtitles, Error, TEXT("Bad WVTT text sample box, ignoring."));
+			break;
+		}
+
+		// An empty cue?
+		if (BoxType == ElectraSubtitleDecoderWVTTUtils::FDataReaderMP4::BoxType_vtte)
+		{
+			if (bInsideCue)
+			{
+				Subtitles.Emplace(MoveTemp(Sub));
+			}
+			bInsideCue = true;
+			// This must be the only entry, so we stop parsing now.
+			// If there are additional boxes this is an authoring error we ignore.
+			break;
+		}
+		// An additional text box (a comment)?
+		else if (BoxType == ElectraSubtitleDecoderWVTTUtils::FDataReaderMP4::BoxType_vtta)
+		{
+			if (bInsideCue)
+			{
+				Subtitles.Emplace(MoveTemp(Sub));
+			}
+			bInsideCue = true;
+			Sub.bIsAdditionalCue = true;
+			RETURN_IF_ERROR(r.ReadString(Sub.Text, BoxLen - 8));
+		}
+		// A cue?
+		else if (BoxType == ElectraSubtitleDecoderWVTTUtils::FDataReaderMP4::BoxType_vttc)
+		{
+			if (bInsideCue)
+			{
+				Subtitles.Emplace(MoveTemp(Sub));
+			}
+			bInsideCue = true;
+		}
+		else if (bInsideCue)
+		{
+			// Cue source ID?
+			if (BoxType == ElectraSubtitleDecoderWVTTUtils::FDataReaderMP4::BoxType_vsid)
+			{
+				int32 source_ID;
+				RETURN_IF_ERROR(r.Read(source_ID));
+				Sub.SourceID = source_ID;
+			}
+			// Cue time?
+			else if (BoxType == ElectraSubtitleDecoderWVTTUtils::FDataReaderMP4::BoxType_ctim)
+			{
+				FString s;
+				RETURN_IF_ERROR(r.ReadString(s, BoxLen - 8));
+				Sub.CurrentTime = s;
+			}
+			// ID?
+			else if (BoxType == ElectraSubtitleDecoderWVTTUtils::FDataReaderMP4::BoxType_iden)
+			{
+				FString s;
+				RETURN_IF_ERROR(r.ReadString(s, BoxLen - 8));
+				Sub.ID = s;
+			}
+			// Settings?
+			else if (BoxType == ElectraSubtitleDecoderWVTTUtils::FDataReaderMP4::BoxType_sttg)
+			{
+				FString s;
+				RETURN_IF_ERROR(r.ReadString(s, BoxLen - 8));
+				Sub.Settings = s;
+			}
+			// Payload?
+			else if (BoxType == ElectraSubtitleDecoderWVTTUtils::FDataReaderMP4::BoxType_payl)
+			{
+				RETURN_IF_ERROR(r.ReadString(Sub.Text, BoxLen - 8));
+			}
+			// Something else.
+			else if (!r.ReadBytes(nullptr, BoxLen - 8))
+			{
+				break;
+			}
+		}
+		// Something else.
+		else if (!r.ReadBytes(nullptr, BoxLen - 8))
+		{
+			break;
+		}
+	}
+	// At the end of the data add the currently worked on cue, if there is one.
+	if (bInsideCue)
+	{
+		Subtitles.Emplace(MoveTemp(Sub));
+	}
+
+	for(auto &Cue : Subtitles)
+	{
+		if (!Cue.bIsAdditionalCue)
+		{
+			TSharedPtr<FSubtitleDecoderOutputWVTT, ESPMode::ThreadSafe> Out = MakeShared<FSubtitleDecoderOutputWVTT, ESPMode::ThreadSafe>();
+			Out->SetTimestamp(InAbsoluteTimestamp);
+			Out->SetDuration(InDuration);
+			// This decoder returns plain text only. Remove all formatting tags.
+			Out->SetText(RemoveAllTags(Cue.Text));
+			if (Cue.ID.IsSet())
+			{
+				Out->SetID(Cue.ID.GetValue());
+			}
+			else
+			{
+				Out->SetID(FString(TEXT("<")) + LexToString(++NextID) + FString(TEXT(">")));
+			}
+			ParsedSubtitleDelegate.Broadcast(Out);
+		}
+	}
+
+	#undef RETURN_IF_ERROR
+}
+
+void FElectraSubtitleDecoderWVTT::SignalStreamedSubtitleEOD()
+{
+}
+
+void FElectraSubtitleDecoderWVTT::Flush()
+{
+}
+
+void FElectraSubtitleDecoderWVTT::Start()
+{
+}
+
+void FElectraSubtitleDecoderWVTT::Stop()
+{
+}
+
+void FElectraSubtitleDecoderWVTT::UpdatePlaybackPosition(Electra::FTimeValue InAbsolutePosition, Electra::FTimeValue InLocalPosition)
+{
+}
+
+
+
+FString FElectraSubtitleDecoderWVTT::RemoveAllTags(const FString& InString)
+{
+	/*
+		For simplicity we do not attempt to parse the tags.
+		Instead we remove everything enclosed in a '<' '>' pair, including the less and greater characters themselves.
+		We also do not check for balanced <>. A sequence like <<x> will also be removed entirely.
+
+		See: https://developer.mozilla.org/en-US/docs/Web/API/WebVTT_API
+	*/
+	FString Result;
+	const TCHAR * const TagStart = TEXT("<");
+	const TCHAR * const TagEnd = TEXT(">");
+	int32 CurrentPos = 0;
+
+	while(CurrentPos < InString.Len())
+	{
+		// Find the next opening tag
+		int32 ltPos = InString.Find(TagStart, ESearchCase::CaseSensitive, ESearchDir::FromStart, CurrentPos);
+		int32 gtPos = INDEX_NONE;
+		if (ltPos != INDEX_NONE)
+		{
+			// Find tag close char. If there is none, we do not remove the opening character.
+			gtPos = InString.Find(TagEnd, ESearchCase::CaseSensitive, ESearchDir::FromStart, ltPos);
+			ltPos = gtPos == INDEX_NONE ? gtPos : ltPos;
+		}
+		Result += InString.Mid(CurrentPos, ltPos == INDEX_NONE ? MAX_int32 : ltPos - CurrentPos);
+		CurrentPos = gtPos != INDEX_NONE ? gtPos+1 : ltPos == INDEX_NONE ? MAX_int32 : ltPos+1;
+	}
+
+	// Replace common XML entities.
+	Result.ReplaceInline(TEXT("&lt;"),   TEXT("<"),  ESearchCase::CaseSensitive);
+	Result.ReplaceInline(TEXT("&gt;"),   TEXT(">"),  ESearchCase::CaseSensitive);
+	Result.ReplaceInline(TEXT("&nbsp;"), TEXT(" "),  ESearchCase::CaseSensitive);
+	Result.ReplaceInline(TEXT("&amp;"),  TEXT("&"),  ESearchCase::CaseSensitive);
+	Result.ReplaceInline(TEXT("&quot;"), TEXT("\""), ESearchCase::CaseSensitive);
+
+	return Result;
+}
+
diff --git a/Engine/Plugins/Media/ElectraSubtitles/Source/ElectraSubtitles/Private/wvtt/ElectraSubtitleDecoder_WVTT.h b/Engine/Plugins/Media/ElectraSubtitles/Source/ElectraSubtitles/Private/wvtt/ElectraSubtitleDecoder_WVTT.h
new file mode 100644
--- /dev/null
+++ b/Engine/Plugins/Media/ElectraSubtitles/Source/ElectraSubtitles/Private/wvtt/ElectraSubtitleDecoder_WVTT.h
@@ -0,0 +1,58 @@
+// Copyright Epic Games, Inc. All Rights Reserved.
+
+#pragma once
+
+#include "ElectraSubtitleDecoder.h"
+#include "ElectraSubtitleDecoderFactory.h"
+
+/**
+ * WebVTT subtitle decoder (https://www.w3.org/TR/webvtt1/)
+ */
+class FElectraSubtitleDecoderWVTT : public IElectraSubtitleDecoder
+{
+public:
+	static void RegisterCodecs(IElectraSubtitleDecoderFactoryRegistry& InRegistry);
+
+	FElectraSubtitleDecoderWVTT();
+
+	virtual ~FElectraSubtitleDecoderWVTT();
+
+	//-------------------------------------------------------------------------
+	// Methods from IElectraSubtitleDecoder
+	//
+	virtual bool InitializeStreamWithCSD(const TArray<uint8>& InCSD, const Electra::FParamDict& InAdditionalInfo) override;
+	virtual IElectraSubtitleDecoder::FOnSubtitleReceivedDelegate& GetParsedSubtitleReceiveDelegate() override;
+	virtual Electra::FTimeValue GetStreamedDeliveryTimeOffset() override;
+	virtual void AddStreamedSubtitleData(const TArray<uint8>& InData, Electra::FTimeValue InAbsoluteTimestamp, Electra::FTimeValue InDuration, const Electra::FParamDict& InAdditionalInfo) override;
+	virtual void SignalStreamedSubtitleEOD() override;
+	virtual void Flush() override;
+	virtual void Start() override;
+	virtual void Stop() override;
+	virtual void UpdatePlaybackPosition(Electra::FTimeValue InAbsolutePosition, Electra::FTimeValue InLocalPosition) override;
+
+private:
+	struct FSubtitleEntry
+	{
+		FString Text;
+		TOptional<int32> SourceID;
+		TOptional<FString> CurrentTime;
+		TOptional<FString> ID;
+		TOptional<FString> Settings;
+		bool bIsAdditionalCue = false;
+	};
+
+	FString RemoveAllTags(const FString& InString);
+
+	FOnSubtitleReceivedDelegate ParsedSubtitleDelegate;
+	uint32 NextID = 0;
+
+	uint16				DataReferenceIndex = 0;
+	int32				Width = 0;
+	int32				Height = 0;
+	int32				TranslationX = 0;
+	int32				TranslationY = 0;
+	uint32				Timescale = 0;
+
+	FString				Configuration;
+	FString				Label;
+};
