Updated ElectraSubtitles to use a shared PCH.

Before: 2.79 seconds
After: 1.66 seconds

diff --git a/Engine/Plugins/Media/ElectraSubtitles/Source/ElectraSubtitles/ElectraSubtitles.Build.cs b/Engine/Plugins/Media/ElectraSubtitles/Source/ElectraSubtitles/ElectraSubtitles.Build.cs
--- a/Engine/Plugins/Media/ElectraSubtitles/Source/ElectraSubtitles/ElectraSubtitles.Build.cs
+++ b/Engine/Plugins/Media/ElectraSubtitles/Source/ElectraSubtitles/ElectraSubtitles.Build.cs
@@ -6,7 +6,9 @@ namespace UnrealBuildTool.Rules
 	{
 		public ElectraSubtitles(ReadOnlyTargetRules Target) : base(Target)
 		{
-			PCHUsage = PCHUsageMode.NoPCHs;
+			PCHUsage = PCHUsageMode.UseSharedPCHs;
+			PrivatePCHHeaderFile = "Public/ElectraSubtitleModule.h";
+
 			PrivateDependencyModuleNames.AddRange(
 				new string[]
 				{
