ElectraSubtitles: Fixed ttml iterator becoming invalid in range-based loop after an emplace

#rnx

diff --git a/Engine/Plugins/Media/ElectraSubtitles/Source/ElectraSubtitles/Private/ttml/TTMLParser.cpp b/Engine/Plugins/Media/ElectraSubtitles/Source/ElectraSubtitles/Private/ttml/TTMLParser.cpp
--- a/Engine/Plugins/Media/ElectraSubtitles/Source/ElectraSubtitles/Private/ttml/TTMLParser.cpp
+++ b/Engine/Plugins/Media/ElectraSubtitles/Source/ElectraSubtitles/Private/ttml/TTMLParser.cpp
@@ -969,7 +969,9 @@ bool FTTMLXML_TTElement::ProcessAttributes(FTTMLParser* Parser, const TArray<FXm
 		}
 		else
 		{
-			Attributes.Emplace(Attr);
+			// Add attribute. Do not use Emplace() because
+			//		V789: Iterators for the 'Attributes' container, used in the range-based for loop, become invalid upon the call of the 'Emplace' function.
+			Attributes.Add(Attr);
 		}
 	}
 	return Super::ProcessAttributes(Parser, Attributes);
