ElectraSubtitles: Fixed iterating the wrong attribute list

#rnx

diff --git a/Engine/Plugins/Media/ElectraSubtitles/Source/ElectraSubtitles/Private/ttml/TTMLParser.cpp b/Engine/Plugins/Media/ElectraSubtitles/Source/ElectraSubtitles/Private/ttml/TTMLParser.cpp
--- a/Engine/Plugins/Media/ElectraSubtitles/Source/ElectraSubtitles/Private/ttml/TTMLParser.cpp
+++ b/Engine/Plugins/Media/ElectraSubtitles/Source/ElectraSubtitles/Private/ttml/TTMLParser.cpp
@@ -882,8 +882,10 @@ bool FTTMLXML_TTElement::ProcessElement(FTTMLParser* Parser, const FString& Elem
 bool FTTMLXML_TTElement::ProcessAttributes(FTTMLParser* Parser, const TArray<FXmlAttribute>& AllAttributes)
 {
 	TArray<FXmlAttribute> Attributes;
-	for(auto& Attr : Attributes)
+	for(int32 i=0; i<AllAttributes.Num(); ++i)
 	{
+		const FXmlAttribute& Attr = Attributes[i];
+
 		// See: https://www.w3.org/TR/ttml2/#parameter-attribute-timeBase
 		if (ElectraSubtitleUtils::StringEquals(*Attr.GetName(), TEXT("timeBase")))
 		{
@@ -969,8 +971,6 @@ bool FTTMLXML_TTElement::ProcessAttributes(FTTMLParser* Parser, const TArray<FXm
 		}
 		else
 		{
-			// Add attribute. Do not use Emplace() because
-			//		V789: Iterators for the 'Attributes' container, used in the range-based for loop, become invalid upon the call of the 'Emplace' function.
 			Attributes.Add(Attr);
 		}
 	}
