ElectraSubtitles: Fixed accessing wrong attribute list

#rnx

diff --git a/Engine/Plugins/Media/ElectraSubtitles/Source/ElectraSubtitles/Private/ttml/TTMLParser.cpp b/Engine/Plugins/Media/ElectraSubtitles/Source/ElectraSubtitles/Private/ttml/TTMLParser.cpp
--- a/Engine/Plugins/Media/ElectraSubtitles/Source/ElectraSubtitles/Private/ttml/TTMLParser.cpp
+++ b/Engine/Plugins/Media/ElectraSubtitles/Source/ElectraSubtitles/Private/ttml/TTMLParser.cpp
@@ -884,7 +884,7 @@ bool FTTMLXML_TTElement::ProcessAttributes(FTTMLParser* Parser, const TArray<FXm
 	TArray<FXmlAttribute> Attributes;
 	for(int32 i=0; i<AllAttributes.Num(); ++i)
 	{
-		const FXmlAttribute& Attr = Attributes[i];
+		const FXmlAttribute& Attr = AllAttributes[i];
 
 		// See: https://www.w3.org/TR/ttml2/#parameter-attribute-timeBase
 		if (ElectraSubtitleUtils::StringEquals(*Attr.GetName(), TEXT("timeBase")))
