ElectraSubtitles: Setting the sequence index on the output

#rnx

diff --git a/Engine/Plugins/Media/ElectraSubtitles/Source/ElectraSubtitles/Private/ttml/ElectraSubtitleDecoder_TTML.cpp b/Engine/Plugins/Media/ElectraSubtitles/Source/ElectraSubtitles/Private/ttml/ElectraSubtitleDecoder_TTML.cpp
--- a/Engine/Plugins/Media/ElectraSubtitles/Source/ElectraSubtitles/Private/ttml/ElectraSubtitleDecoder_TTML.cpp
+++ b/Engine/Plugins/Media/ElectraSubtitles/Source/ElectraSubtitles/Private/ttml/ElectraSubtitleDecoder_TTML.cpp
@@ -61,7 +61,7 @@ public:
 	void SetTimestamp(const Electra::FTimeValue& InTimestamp)
 	{
 		Timestamp.Time = InTimestamp.GetAsTimespan();
-		Timestamp.SequenceIndex = 0;
+		Timestamp.SequenceIndex = InTimestamp.GetSequenceIndex();
 	}
 
 	void SetID(const FString& InID)
diff --git a/Engine/Plugins/Media/ElectraSubtitles/Source/ElectraSubtitles/Private/tx3g/ElectraSubtitleDecoder_TX3G.cpp b/Engine/Plugins/Media/ElectraSubtitles/Source/ElectraSubtitles/Private/tx3g/ElectraSubtitleDecoder_TX3G.cpp
--- a/Engine/Plugins/Media/ElectraSubtitles/Source/ElectraSubtitles/Private/tx3g/ElectraSubtitleDecoder_TX3G.cpp
+++ b/Engine/Plugins/Media/ElectraSubtitles/Source/ElectraSubtitles/Private/tx3g/ElectraSubtitleDecoder_TX3G.cpp
@@ -189,7 +189,7 @@ public:
 	void SetTimestamp(const Electra::FTimeValue& InTimestamp)
 	{
 		Timestamp.Time = InTimestamp.GetAsTimespan();
-		Timestamp.SequenceIndex = 0;
+		Timestamp.SequenceIndex = InTimestamp.GetSequenceIndex();
 	}
 
 	void SetID(const FString& InID)
diff --git a/Engine/Plugins/Media/ElectraSubtitles/Source/ElectraSubtitles/Private/wvtt/ElectraSubtitleDecoder_WVTT.cpp b/Engine/Plugins/Media/ElectraSubtitles/Source/ElectraSubtitles/Private/wvtt/ElectraSubtitleDecoder_WVTT.cpp
--- a/Engine/Plugins/Media/ElectraSubtitles/Source/ElectraSubtitles/Private/wvtt/ElectraSubtitleDecoder_WVTT.cpp
+++ b/Engine/Plugins/Media/ElectraSubtitles/Source/ElectraSubtitles/Private/wvtt/ElectraSubtitleDecoder_WVTT.cpp
@@ -190,7 +190,7 @@ public:
 	void SetTimestamp(const Electra::FTimeValue& InTimestamp)
 	{
 		Timestamp.Time = InTimestamp.GetAsTimespan();
-		Timestamp.SequenceIndex = 0;
+		Timestamp.SequenceIndex = InTimestamp.GetSequenceIndex();
 	}
 
 	void SetID(const FString& InID)
