Electra: Added webvtt subtitle parser; enabled subtitles with HLS; fixed issue with selecting initial subtitle track

#rnx

diff --git a/Engine/Plugins/Media/ElectraPlayer/Source/ElectraPlayerRuntime/Private/Runtime/Decoder/SubtitleDecoder.cpp b/Engine/Plugins/Media/ElectraPlayer/Source/ElectraPlayerRuntime/Private/Runtime/Decoder/SubtitleDecoder.cpp
--- a/Engine/Plugins/Media/ElectraPlayer/Source/ElectraPlayerRuntime/Private/Runtime/Decoder/SubtitleDecoder.cpp
+++ b/Engine/Plugins/Media/ElectraPlayer/Source/ElectraPlayerRuntime/Private/Runtime/Decoder/SubtitleDecoder.cpp
@@ -11,7 +11,7 @@ namespace Electra
 namespace SubtitleOptionKeys
 {
 	static const FName SendEmptySubtitleDuringGaps(TEXT("sendEmptySubtitleDuringGaps"));
-	static const FName SideloadedID(TEXT("SideloadedID"));
+	static const FName SourceID(TEXT("source_id"));
 	static const FName Width(TEXT("width"));
 	static const FName Height(TEXT("height"));
 	static const FName OffsetX(TEXT("offset_x"));
@@ -144,13 +144,21 @@ public:
 	virtual void AUdataClearEOD() override;
 	virtual void AUdataFlushEverything() override;
 
+	void SetIsSideloaded(bool bInIsSideloaded)
+	{
+		CurrentConfig.bIsSideloaded = bInIsSideloaded;
+	}
+	void SetSourceID(const FString& InID)
+	{
+		CurrentConfig.SourceID = InID;
+	}
 	void SetCSD(const TArray<uint8> InRawCSD)
 	{
-		RawCSD = InRawCSD;
+		CurrentConfig.RawCSD = InRawCSD;
 	}
 	void SetCodecInfo(const FStreamCodecInformation& InCodecInfo)
 	{
-		CodecInfo = InCodecInfo;
+		CurrentConfig.CodecInfo = InCodecInfo;
 	}
 
 	bool ResetForTrackChange(FAccessUnit* AccessUnit);
@@ -162,13 +170,27 @@ private:
 	TSharedPtr<IElectraSubtitleDecoder, ESPMode::ThreadSafe> PluginDecoder;
 	FParamDict PluginDecoderOptions;
 	IPlayerSessionServices* PlayerSessionServices = nullptr;
-	TArray<uint8> RawCSD;
-	FStreamCodecInformation CodecInfo;
 	FTimeValue DeliveryOffset;
 	FDelegateHandle ReceiveDelegateHandle;
 	FDecodedSubtitleReceivedDelegate SubtitleReceiverDelegate;
 	FDecodedSubtitleFlushDelegate SubtitleFlushDelegate;
 	bool bIsStarted = false;
+
+	struct FCurrentConfig
+	{
+		FStreamCodecInformation CodecInfo;
+		TArray<uint8> RawCSD;
+		FString SourceID;
+		bool bIsSideloaded = false;
+		void Reset()
+		{
+			CodecInfo.Clear();
+			RawCSD.Empty();
+			SourceID.Empty();
+			bIsSideloaded = false;
+		}
+	};
+	FCurrentConfig CurrentConfig;
 };
 
 bool ISubtitleDecoder::IsSupported(const FString& MimeType, const FString& Codec)
@@ -195,6 +217,8 @@ ISubtitleDecoder* ISubtitleDecoder::Create(FAccessUnit* AccessUnit)
 		if (PluginDecoder.IsValid())
 		{
 			FSubtitleDecoder* Decoder = new FSubtitleDecoder(MoveTemp(PluginDecoder));
+			Decoder->SetIsSideloaded(AccessUnit->bIsSideloaded);
+			Decoder->SetSourceID(AccessUnit->BufferSourceInfo.IsValid() ? AccessUnit->BufferSourceInfo->PeriodAdaptationSetID : FString());
 			Decoder->SetCSD(AccessUnit->AUCodecData->RawCSD);
 			Decoder->SetCodecInfo(AccessUnit->AUCodecData->ParsedInfo);
 			return Decoder;
@@ -239,12 +263,12 @@ void FSubtitleDecoder::Open(const FParamDict& Options)
 	{
 		PluginDecoderOptions = Options;
 		FParamDict Addtl(Options);
-		Addtl.Set(SubtitleOptionKeys::Width, FVariantValue((int64) CodecInfo.GetResolution().Width));
-		Addtl.Set(SubtitleOptionKeys::Height, FVariantValue((int64) CodecInfo.GetResolution().Height));
-		Addtl.Set(SubtitleOptionKeys::OffsetX, FVariantValue((int64) CodecInfo.GetTranslation().GetX()));
-		Addtl.Set(SubtitleOptionKeys::OffsetY, FVariantValue((int64) CodecInfo.GetTranslation().GetY()));
-		Addtl.Set(SubtitleOptionKeys::Timescale, FVariantValue((int64) CodecInfo.GetFrameRate().GetDenominator()));
-		if (PluginDecoder->InitializeStreamWithCSD(RawCSD, Addtl))
+		Addtl.Set(SubtitleOptionKeys::Width, FVariantValue((int64) CurrentConfig.CodecInfo.GetResolution().Width));
+		Addtl.Set(SubtitleOptionKeys::Height, FVariantValue((int64) CurrentConfig.CodecInfo.GetResolution().Height));
+		Addtl.Set(SubtitleOptionKeys::OffsetX, FVariantValue((int64) CurrentConfig.CodecInfo.GetTranslation().GetX()));
+		Addtl.Set(SubtitleOptionKeys::OffsetY, FVariantValue((int64) CurrentConfig.CodecInfo.GetTranslation().GetY()));
+		Addtl.Set(SubtitleOptionKeys::Timescale, FVariantValue((int64) CurrentConfig.CodecInfo.GetFrameRate().GetDenominator()));
+		if (PluginDecoder->InitializeStreamWithCSD(CurrentConfig.RawCSD, Addtl))
 		{
 			DeliveryOffset = PluginDecoder->GetStreamedDeliveryTimeOffset();
 			ReceiveDelegateHandle = PluginDecoder->GetParsedSubtitleReceiveDelegate().AddRaw(this, &FSubtitleDecoder::OnDecodedSubtitleReceived);
@@ -269,8 +293,7 @@ void FSubtitleDecoder::Close()
 		ReceiveDelegateHandle.Reset();
 	}
 	PluginDecoder.Reset();
-	RawCSD.Empty();
-	CodecInfo.Clear();
+	CurrentConfig.Reset();
 }
 
 void FSubtitleDecoder::Start()
@@ -317,7 +340,14 @@ ISubtitleDecoder::FDecodedSubtitleFlushDelegate& FSubtitleDecoder::GetDecodedSub
 
 void FSubtitleDecoder::OnDecodedSubtitleReceived(ISubtitleDecoderOutputPtr DecodedSubtitle)
 {
-	SubtitleReceiverDelegate.ExecuteIfBound(DecodedSubtitle);
+	if (DecodedSubtitle.IsValid())
+	{
+	/*
+		FDecoderTimeStamp ts = DecodedSubtitle->GetTime();
+		DecodedSubtitle->SetTime(ts);
+	*/
+		SubtitleReceiverDelegate.ExecuteIfBound(DecodedSubtitle);
+	}
 }
 
 
@@ -325,16 +355,52 @@ bool FSubtitleDecoder::ResetForTrackChange(FAccessUnit* AccessUnit)
 {
 	bool bWasStarted = bIsStarted;
 	Stop();
+	if (PluginDecoder.IsValid())
+	{
+		PluginDecoder->Flush();
+	}
 	SubtitleFlushDelegate.ExecuteIfBound();
-	Close();
+
+	// Do we need to close and create a new decoder?
 	if (AccessUnit && AccessUnit->AUCodecData.IsValid())
 	{
-		PluginDecoder = FSubtitleDecoderPlugins::Get().CreateDecoder(AccessUnit->AUCodecData->ParsedInfo.GetCodecSpecifierRFC6381());
-		if (PluginDecoder.IsValid())
+		FString NewSourceID(AccessUnit->BufferSourceInfo.IsValid() ? AccessUnit->BufferSourceInfo->PeriodAdaptationSetID : FString());
+		if (CurrentConfig.SourceID != NewSourceID ||
+			CurrentConfig.bIsSideloaded != AccessUnit->bIsSideloaded ||
+			CurrentConfig.RawCSD != AccessUnit->AUCodecData->RawCSD ||
+			!CurrentConfig.CodecInfo.Equals(AccessUnit->AUCodecData->ParsedInfo))
+		{
+			Close();
+			if (AccessUnit && AccessUnit->AUCodecData.IsValid())
+			{
+				// For sideloaded data prefer to create the decoder via the Mime type. This allows the decoder to
+				// know that the data is there all at once and it needs to handle it in a non-streaming fashion.
+				if (AccessUnit->bIsSideloaded)
+				{
+					PluginDecoder = FSubtitleDecoderPlugins::Get().CreateDecoder(AccessUnit->AUCodecData->ParsedInfo.GetMimeType());
+				}
+				if (!PluginDecoder.IsValid())
+				{
+					PluginDecoder = FSubtitleDecoderPlugins::Get().CreateDecoder(AccessUnit->AUCodecData->ParsedInfo.GetCodecSpecifierRFC6381());
+				}
+				if (PluginDecoder.IsValid())
+				{
+					SetIsSideloaded(AccessUnit->bIsSideloaded);
+					SetSourceID(NewSourceID);
+					SetCSD(AccessUnit->AUCodecData->RawCSD);
+					SetCodecInfo(AccessUnit->AUCodecData->ParsedInfo);
+					Open(PluginDecoderOptions);
+					if (bWasStarted)
+					{
+						Start();
+					}
+					return true;
+				}
+			}
+		}
+		else
 		{
-			SetCSD(AccessUnit->AUCodecData->RawCSD);
-			SetCodecInfo(AccessUnit->AUCodecData->ParsedInfo);
-			Open(PluginDecoderOptions);
+			// No change needed.
 			if (bWasStarted)
 			{
 				Start();
@@ -359,12 +425,12 @@ void FSubtitleDecoder::AUdataPushAU(FAccessUnit* AccessUnit)
 					return;
 				}
 			}
-			FParamDict Addtl;
-			// For sideloaded subtitles we pass an ID to the plugin via the additional options. This allows the plugin to identify
+			// Set the period and adaptation set ID in the additional options. This allows the plugin to identify
 			// whether or not it has already parsed this data before (when seeking for instance).
-			if (AccessUnit->bIsSideloaded && AccessUnit->BufferSourceInfo.IsValid())
+			FParamDict Addtl;
+			if (AccessUnit->BufferSourceInfo.IsValid())
 			{
-				Addtl.Set(SubtitleOptionKeys::SideloadedID, FVariantValue(AccessUnit->BufferSourceInfo->PeriodAdaptationSetID));
+				Addtl.Set(SubtitleOptionKeys::SourceID, FVariantValue(AccessUnit->BufferSourceInfo->PeriodAdaptationSetID));
 			}
 			PluginDecoder->AddStreamedSubtitleData(TArray<uint8>((const uint8*)AccessUnit->AUData, (int32)AccessUnit->AUSize), AccessUnit->PTS, AccessUnit->Duration, Addtl);
 		}
diff --git a/Engine/Plugins/Media/ElectraPlayer/Source/ElectraPlayerRuntime/Private/Runtime/Player/AdaptiveStreamingPlayer.cpp b/Engine/Plugins/Media/ElectraPlayer/Source/ElectraPlayerRuntime/Private/Runtime/Player/AdaptiveStreamingPlayer.cpp
--- a/Engine/Plugins/Media/ElectraPlayer/Source/ElectraPlayerRuntime/Private/Runtime/Player/AdaptiveStreamingPlayer.cpp
+++ b/Engine/Plugins/Media/ElectraPlayer/Source/ElectraPlayerRuntime/Private/Runtime/Player/AdaptiveStreamingPlayer.cpp
@@ -3375,7 +3375,7 @@ void FAdaptiveStreamingPlayer::HandlePendingMediaSegmentRequests()
 								case EStreamType::Video:
 								{
 									SelectedStreamAttributesVid.UpdateWith(BufferSourceInfo->Kind, BufferSourceInfo->LanguageTag, BufferSourceInfo->Codec, BufferSourceInfo->HardIndex);
-									if (ManifestType != EMediaFormatType::ISOBMFF && ManifestType != EMediaFormatType::MKV)
+									if (ManifestType == EMediaFormatType::DASH)
 									{
 										StreamSelectionAttributesVid.ClearOverrideIndex();
 									}
@@ -3384,7 +3384,7 @@ void FAdaptiveStreamingPlayer::HandlePendingMediaSegmentRequests()
 								case EStreamType::Audio:
 								{
 									SelectedStreamAttributesAud.UpdateWith(BufferSourceInfo->Kind, BufferSourceInfo->LanguageTag, BufferSourceInfo->Codec, BufferSourceInfo->HardIndex);
-									if (ManifestType != EMediaFormatType::ISOBMFF && ManifestType != EMediaFormatType::MKV)
+									if (ManifestType == EMediaFormatType::DASH)
 									{
 										StreamSelectionAttributesAud.ClearOverrideIndex();
 									}
@@ -3393,7 +3393,7 @@ void FAdaptiveStreamingPlayer::HandlePendingMediaSegmentRequests()
 								case EStreamType::Subtitle:
 								{
 									SelectedStreamAttributesTxt.UpdateWith(BufferSourceInfo->Kind, BufferSourceInfo->LanguageTag, BufferSourceInfo->Codec, BufferSourceInfo->HardIndex);
-									if (ManifestType != EMediaFormatType::ISOBMFF && ManifestType != EMediaFormatType::MKV)
+									if (ManifestType == EMediaFormatType::DASH)
 									{
 										StreamSelectionAttributesTxt.ClearOverrideIndex();
 									}
@@ -3736,10 +3736,7 @@ void FAdaptiveStreamingPlayer::InternalHandleSegmentTrackChanges(const FTimeValu
 
 		HandleElementaryStreamTrackSelection(StreamSelectionAttributesVid, SelectedStreamAttributesVid, EStreamType::Video, PendingTrackSelectionVid, CurrentPlayPeriodVideo);
 		HandleElementaryStreamTrackSelection(StreamSelectionAttributesAud, SelectedStreamAttributesAud, EStreamType::Audio, PendingTrackSelectionAud, CurrentPlayPeriodAudio);
-		if (ManifestType == EMediaFormatType::DASH)
-		{
-			HandleElementaryStreamTrackSelection(StreamSelectionAttributesTxt, SelectedStreamAttributesTxt, EStreamType::Subtitle, PendingTrackSelectionTxt, CurrentPlayPeriodText);
-		}
+		HandleElementaryStreamTrackSelection(StreamSelectionAttributesTxt, SelectedStreamAttributesTxt, EStreamType::Subtitle, PendingTrackSelectionTxt, CurrentPlayPeriodText);
 	}
 	else
 	{
diff --git a/Engine/Plugins/Media/ElectraPlayer/Source/ElectraPlayerRuntime/Private/Runtime/Player/HLS/PlaylistHLS.cpp b/Engine/Plugins/Media/ElectraPlayer/Source/ElectraPlayerRuntime/Private/Runtime/Player/HLS/PlaylistHLS.cpp
--- a/Engine/Plugins/Media/ElectraPlayer/Source/ElectraPlayerRuntime/Private/Runtime/Player/HLS/PlaylistHLS.cpp
+++ b/Engine/Plugins/Media/ElectraPlayer/Source/ElectraPlayerRuntime/Private/Runtime/Player/HLS/PlaylistHLS.cpp
@@ -142,7 +142,7 @@ private:
 	IPlayerSessionServices* PlayerSessionServices = nullptr;
 	TSharedPtr<FTimelineMediaAsset, ESPMode::ThreadSafe> TimelineMediaAsset;
 	IManifest::IPlayPeriod::EReadyState CurrentReadyState = IManifest::IPlayPeriod::EReadyState::NotLoaded;
-	TOptional<FStreamSelectionAttributes> StreamSelectionAttributes[4];
+	FStreamSelectionAttributes StreamSelectionAttributes[4];
 	FSelectedTrackStream SelectedTrackStream[4];
 };
 
@@ -796,7 +796,7 @@ FErrorDetail FActiveHLSPlaylist::CreateTrackMetadata(IPlayerSessionServices* InP
 					tm.VideoVariantBaseIDs.Emplace(vvIt.ID);
 				}
 				tm.Rendition = rg->Renditions[i];
-				tm.Meta.ID = rg->Renditions[i].StableRenditionId;
+				tm.Meta.ID = FString::Printf(TEXT("sub:%s:%s"), *SubtitleGroups[0], *rg->Renditions[i].Name);
 				tm.Meta.Label = rg->Renditions[i].Name;
 				tm.Meta.LanguageTagRFC5646 = rg->Renditions[i].LanguageRFC5646;
 				tm.Meta.Kind = TEXT("subtitles");
@@ -2485,14 +2485,18 @@ void FActiveHLSPlaylist::FPlayPeriod::PrepareForPlay()
 		}
 		else if (nStreamTypeIdx == 2 && Pwy->SubtitleTracks.Num())
 		{
-		// Not setting up subtitle tracks for now as we have no WebVTT parser.
-			//Tracks = &Pwy->SubtitleTracks;
+			// Subtitle tracks are not selected by default.
+			// If there is no explicit selection asked for we ignore them.
+			if (StreamSelectionAttributes[nStreamTypeIdx].IsSet())
+			{
+				Tracks = &Pwy->SubtitleTracks;
+			}
 		}
 
 		int32 SelectedTrackIndex = 0;
-		if (Tracks && StreamSelectionAttributes[nStreamTypeIdx].IsSet() && StreamSelectionAttributes[nStreamTypeIdx].GetValue().IsSet())
+		if (Tracks && StreamSelectionAttributes[nStreamTypeIdx].IsSet())
 		{
-			const FStreamSelectionAttributes& Sel(StreamSelectionAttributes[nStreamTypeIdx].GetValue());
+			const FStreamSelectionAttributes& Sel(StreamSelectionAttributes[nStreamTypeIdx]);
 
 			// Is this a hard choice?
 			if (Sel.OverrideIndex.IsSet() && Sel.OverrideIndex.GetValue() < Tracks->Num())
diff --git a/Engine/Plugins/Media/ElectraPlayer/Source/ElectraPlayerRuntime/Private/Runtime/Player/StreamSegmentReaderCommon.cpp b/Engine/Plugins/Media/ElectraPlayer/Source/ElectraPlayerRuntime/Private/Runtime/Player/StreamSegmentReaderCommon.cpp
--- a/Engine/Plugins/Media/ElectraPlayer/Source/ElectraPlayerRuntime/Private/Runtime/Player/StreamSegmentReaderCommon.cpp
+++ b/Engine/Plugins/Media/ElectraPlayer/Source/ElectraPlayerRuntime/Private/Runtime/Player/StreamSegmentReaderCommon.cpp
@@ -760,14 +760,8 @@ void FStreamSegmentReaderCommon::FStreamHandler::HandleRequest()
 			{
 				FScopeLock ProbeBufferLock(ReadBuffer.GetLock());
 				const uint8* MediaDataPtr = ReadBuffer.GetLinearReadData();
-				if (ReadBuffer.GetLinearReadSize() >= kNumProbeBytesNeeded)
-				{
-					FMemory::Memcpy(ProbeBytes, ReadBuffer.GetLinearReadData(), kNumProbeBytesNeeded);
-				}
-				else
-				{
-					bContinue = false;
-				}
+				int32 NumProbeBytes = Utils::Min((int32) ReadBuffer.GetLinearReadSize(), (int32)kNumProbeBytesNeeded);
+				FMemory::Memcpy(ProbeBytes, ReadBuffer.GetLinearReadData(), NumProbeBytes);
 			}
 		}
 
@@ -802,6 +796,13 @@ void FStreamSegmentReaderCommon::FStreamHandler::HandleRequest()
 				ID3HeaderSize = (int32) (10U + (ProbeBytes[6] << 21U) + (ProbeBytes[7] << 14U) + (ProbeBytes[8] << 7U) + ProbeBytes[9]);
 				ContainerFormat = FStreamSegmentRequestCommon::EContainerFormat::ID3_Raw;
 			}
+			// Raw WebVTT file ? (starts with "WEBVTT" or BOM+"WEBVTT")
+			// Could also be any even empty file as long as the init segment indicates WebVTT
+			else if (IsWebVTTHeader(ProbeBytes, 9) ||
+				(InitSegmentData.IsType<TSharedPtrTS<const TArray<uint8>>>() && IsWebVTTHeader(InitSegmentData.Get<TSharedPtrTS<const TArray<uint8>>>()->GetData(), InitSegmentData.Get<TSharedPtrTS<const TArray<uint8>>>()->Num())))
+			{
+				ContainerFormat = FStreamSegmentRequestCommon::EContainerFormat::WebVTT_Raw;
+			}
 			// ISO/IEC 13818-1 ?
 			else if (ProbeBytes[0] == 0x47)
 			{
@@ -824,6 +825,10 @@ void FStreamSegmentReaderCommon::FStreamHandler::HandleRequest()
 			{
 				SegmentHandleResult = HandleID3RawMedia(Request, ID3HeaderSize);
 			}
+			else if (ContainerFormat == FStreamSegmentRequestCommon::EContainerFormat::WebVTT_Raw)
+			{
+				SegmentHandleResult = HandleRawSubtitleMedia(Request);
+			}
 			else
 			{
 				if (ReadBuffer.DidDecryptionFail())
@@ -922,7 +927,6 @@ bool FStreamSegmentReaderCommon::FStreamHandler::FetchInitSegment(TSharedPtrTS<F
 		return true;
 	}
 
-	FStreamSegmentRequestCommon::EContainerFormat ContainerFormat = InRequest->ExpectedContainerFormat;
 	// Otherwise, check with the entity cache if we already have it from earlier.
 	TSharedPtrTS<IPlayerEntityCache> EntityCache = PlayerSessionService ? PlayerSessionService->GetEntityCache() : nullptr;
 	if (EntityCache.IsValid())
@@ -933,25 +937,23 @@ bool FStreamSegmentReaderCommon::FStreamHandler::FetchInitSegment(TSharedPtrTS<F
 			if (CachedItem.Parsed14496_12Data.IsValid())
 			{
 				InitSegmentData.Emplace<TSharedPtrTS<const IParserISO14496_12>>(CachedItem.Parsed14496_12Data);
-				ContainerFormat =  FStreamSegmentRequestCommon::EContainerFormat::ISO_14496_12;
 				DownloadStats.bWasSuccessful = true;
 			}
 			else if (CachedItem.ParsedMatroskaData.IsValid())
 			{
 				InitSegmentData.Emplace<TSharedPtrTS<const IParserMKV>>(CachedItem.ParsedMatroskaData);
-				ContainerFormat =  FStreamSegmentRequestCommon::EContainerFormat::Matroska_WebM;
 				DownloadStats.bWasSuccessful = true;
 			}
 			else if (CachedItem.RawPayloadData.IsValid())
 			{
 				InitSegmentData.Emplace<TSharedPtrTS<const TArray<uint8>>>(CachedItem.RawPayloadData);
-				ContainerFormat =  FStreamSegmentRequestCommon::EContainerFormat::ISO_13818_1;
 				DownloadStats.bWasSuccessful = true;
 			}
 			return true;
 		}
 	}
 
+	FStreamSegmentRequestCommon::EContainerFormat ContainerFormat = InRequest->ExpectedContainerFormat;
 	// Not cached yet, need to fetch.
 	CurrentConnectionInfo = HTTP::FConnectionInfo();
 	FGenericDataReader StaticSegmentDataReader;
@@ -998,7 +1000,7 @@ bool FStreamSegmentReaderCommon::FStreamHandler::FetchInitSegment(TSharedPtrTS<F
 			return false;
 		}
 		// We need to have some amount of data that this could even be an init segment of sorts.
-		if (ResponseBuffer->Num() < 64)
+		if (ResponseBuffer->Num() < 6)
 		{
 			// Not enough data
 			SetError(FString::Printf(TEXT("Init segment too small to contain relevant data")), INTERNAL_SEG_ERROR_INIT_SEGMENT_TOO_SHORT);
@@ -1069,30 +1071,37 @@ bool FStreamSegmentReaderCommon::FStreamHandler::FetchInitSegment(TSharedPtrTS<F
 
 		// What format did we get?
 		const uint8* InitDataPtr = ResponseBuffer->GetLinearReadData();
+		int64 InitDataSize = ResponseBuffer->GetLinearReadSize();
 		auto GetUINT32BE = [](const uint8* InData) -> uint32
 		{
 			return (static_cast<uint32>(InData[0]) << 24) | (static_cast<uint32>(InData[1]) << 16) | (static_cast<uint32>(InData[2]) << 8) | static_cast<uint32>(InData[3]);
 		};
 		// ISO/IEC 14496-12 ?
-		const uint32 boxName = GetUINT32BE(InitDataPtr + 4);
+		const uint32 boxName = InitDataSize >= 8 ? GetUINT32BE(InitDataPtr + 4) : 0;
 		if (boxName == UtilsMP4::Make4CC('f','t','y','p') || boxName == UtilsMP4::Make4CC('f','r','e','e') || boxName == UtilsMP4::Make4CC('s','k','i','p')
 			|| boxName == UtilsMP4::Make4CC('p','d','i','n') || boxName == UtilsMP4::Make4CC('u','d','t','a'))
 		{
 			StaticSegmentDataReader.SetSourceBuffer(ResponseBuffer);
-			ContainerFormat =  FStreamSegmentRequestCommon::EContainerFormat::ISO_14496_12;
+			ContainerFormat = FStreamSegmentRequestCommon::EContainerFormat::ISO_14496_12;
 		}
 		// Matroska / WebM ?
 		else if (GetUINT32BE(InitDataPtr) == 0x1a45dfa3)
 		{
 			StaticSegmentDataReader.SetSourceBuffer(ResponseBuffer);
-			ContainerFormat =  FStreamSegmentRequestCommon::EContainerFormat::Matroska_WebM;
+			ContainerFormat = FStreamSegmentRequestCommon::EContainerFormat::Matroska_WebM;
 		}
 		// ISO/IEC 13818-1 ?
 		else if (*InitDataPtr == 0x47)
 		{
 			check((ResponseBuffer->GetLinearReadSize() % 188) == 0);
 			StaticSegmentDataReader.SetSourceBuffer(ResponseBuffer);
-			ContainerFormat =  FStreamSegmentRequestCommon::EContainerFormat::ISO_13818_1;
+			ContainerFormat = FStreamSegmentRequestCommon::EContainerFormat::ISO_13818_1;
+		}
+		// Raw WebVTT? (starts with "WEBVTT" or BOM+"WEBVTT")
+		else if (IsWebVTTHeader(InitDataPtr, InitDataSize))
+		{
+			StaticSegmentDataReader.SetSourceBuffer(ResponseBuffer);
+			ContainerFormat = FStreamSegmentRequestCommon::EContainerFormat::WebVTT_Raw;
 		}
 		else
 		{
@@ -1263,6 +1272,23 @@ bool FStreamSegmentReaderCommon::FStreamHandler::FetchInitSegment(TSharedPtrTS<F
 			SetError(FString::Printf(TEXT("Track preparation of init segment \"%s\" failed"), *InitURLToUse.Url.URL), INTERNAL_SEG_ERROR_INIT_SEGMENT_PARSE_ERROR);
 			return false;
 		}
+		case FStreamSegmentRequestCommon::EContainerFormat::WebVTT_Raw:
+		{
+			TSharedPtrTS<TArray<uint8>> RawPacketBlob = MakeSharedTS<TArray<uint8>>(StaticSegmentDataReader.GetBufferBaseAddress(), StaticSegmentDataReader.GetTotalSize());
+			InitSegmentData.Emplace<TSharedPtrTS<const TArray<uint8>>>(RawPacketBlob);
+			// Add this to the entity cache in case it needs to be retrieved again.
+			if (EntityCache.IsValid())
+			{
+				IPlayerEntityCache::FCacheItem CacheItem;
+				CacheItem.URL = InitURLToUse.Url.URL;
+				CacheItem.Range = InitURLToUse.Range;
+				CacheItem.RawPayloadData = RawPacketBlob;
+				EntityCache->CacheEntity(CacheItem);
+			}
+			DownloadStats.bWasSuccessful = true;
+			StreamSelector->ReportDownloadEnd(DownloadStats);
+			return true;
+		}
 		default:
 		{
 			break;
@@ -3100,7 +3126,7 @@ FStreamSegmentReaderCommon::FStreamHandler::EHandleResult FStreamSegmentReaderCo
 					if (!CurrentlyActiveTrackData->CSD.IsValid())
 					{
 						CurrentlyActiveTrackData->CSD = MakeSharedTS<FAccessUnit::CodecData>();
-						FStreamCodecInformation&ci = CurrentlyActiveTrackData->CSD->ParsedInfo;
+						FStreamCodecInformation& ci = CurrentlyActiveTrackData->CSD->ParsedInfo;
 
 						const int32 FrameSize = ElectraDecodersUtil::MPEG::UtilsMPEG123::GetSamplesPerFrame(HeaderValue);
 						const int32 SampleRate = ElectraDecodersUtil::MPEG::UtilsMPEG123::GetSamplingRate(HeaderValue);
@@ -3497,6 +3523,127 @@ FStreamSegmentReaderCommon::FStreamHandler::EHandleResult FStreamSegmentReaderCo
 	return EHandleResult::Finished;
 }
 
+FStreamSegmentReaderCommon::FStreamHandler::EHandleResult FStreamSegmentReaderCommon::FStreamHandler::HandleRawSubtitleMedia(TSharedPtrTS<FStreamSegmentRequestCommon> InRequest)
+{
+	// This is for subtitle segments only
+	if (InRequest->GetType() != EStreamType::Subtitle)
+	{
+		SetError(FString::Printf(TEXT("Raw subtitle streams are expected to be used for subtitles only")), INTERNAL_SEG_ERROR_BAD_MEDIA_SEGMENT);
+		bHasErrored = true;
+		return EHandleResult::Failed;
+	}
+	// And also only for WebVTT at the moment
+	if (!InRequest->CodecInfo[StreamTypeToArrayIndex(EStreamType::Subtitle)].GetCodecName().Equals(TEXT("wvtt"), ESearchCase::IgnoreCase))
+	{
+		SetError(FString::Printf(TEXT("Raw subtitle streams are currently supported fro WebVTT only")), INTERNAL_SEG_ERROR_BAD_MEDIA_SEGMENT);
+		bHasErrored = true;
+		return EHandleResult::Failed;
+	}
+
+	// Create a track map with one entry as track #0
+	TSortedMap<uint64, TSharedPtrTS<FActiveTrackData>> NewTrackDataMap;
+	TSharedPtrTS<FActiveTrackData> td = MakeSharedTS<FActiveTrackData>();
+	td->StreamType = EStreamType::Subtitle;
+	td->bNeedToRecalculateDurations = false;
+	check(InRequest->SourceBufferInfo[StreamTypeToArrayIndex(EStreamType::Subtitle)].IsValid());
+	td->BufferSourceInfo = MakeSharedTS<FBufferSourceInfo>(*(InRequest->SourceBufferInfo[StreamTypeToArrayIndex(EStreamType::Subtitle)]));
+	td->BufferSourceInfo->PlaybackSequenceID = InRequest->GetPlaybackSequenceID();
+	NewTrackDataMap.Emplace(0, MoveTemp(td));
+	TrackDataMap = MoveTemp(NewTrackDataMap);
+	CurrentlyActiveTrackData = TrackDataMap[0];
+	SelectPrimaryTrackData(InRequest);
+
+	// Read the data, but at most 2 MiB
+	const int64 kMaxReadSize = 2 << 20;
+	bool bContinue = true;
+	while(!ReadBuffer.WaitUntilSizeAvailable(kMaxReadSize, 1000*20))
+	{
+		if (HasErrored() || HasReadBeenAborted())
+		{
+			bContinue = false;
+			break;
+		}
+	}
+	if (bContinue)
+	{
+		// Create CSD
+		if (!CurrentlyActiveTrackData->CSD.IsValid())
+		{
+			CurrentlyActiveTrackData->CSD = MakeSharedTS<FAccessUnit::CodecData>();
+			FStreamCodecInformation& ci = CurrentlyActiveTrackData->CSD->ParsedInfo;
+			ci.SetStreamType(EStreamType::Subtitle);
+			ci.SetMimeType(TEXT("text/vtt"));
+			ci.SetCodec(FStreamCodecInformation::ECodec::WebVTT);
+			ci.SetCodecSpecifierRFC6381(TEXT("wvtt"));
+		}
+
+		FTimeValue PTS;
+		PTS.SetFromHNS(InRequest->Segment.MediaLocalFirstAUTime);
+		FTimeValue EarliestPTS(PTS);
+		if (InRequest->Segment.bFrameAccuracyRequired && InRequest->FrameAccurateStartTime.IsValid())
+		{
+			EarliestPTS = InRequest->FrameAccurateStartTime;
+		}
+		EarliestPTS.SetSequenceIndex(InRequest->TimestampSequenceIndex);
+		FTimeValue LatestPTS;
+		if (InRequest->Segment.MediaLocalLastAUTime != TNumericLimits<int64>::Max())
+		{
+			LatestPTS.SetFromHNS(InRequest->Segment.MediaLocalLastAUTime);
+		}
+		else
+		{
+			LatestPTS.SetToPositiveInfinity();
+		}
+		LatestPTS.SetSequenceIndex(InRequest->TimestampSequenceIndex);
+
+		CurrentlyActiveTrackData->DefaultDurationFromCSD.SetFromND(InRequest->Segment.Duration, InRequest->Segment.Timescale);
+
+		// Check if there has been an init segment for this which we need to prepend to the data.
+		TArray<uint8> PrependData;
+		if (InitSegmentData.IsType<TSharedPtrTS<const TArray<uint8>>>())
+		{
+			TSharedPtrTS<const TArray<uint8>> InitData = InitSegmentData.Get<TSharedPtrTS<const TArray<uint8>>>();
+			if (InitData.IsValid() && InitData->Num())
+			{
+				PrependData = *InitData;
+			}
+		}
+
+		// Create the access unit.
+		FAccessUnit *AccessUnit = FAccessUnit::Create(Parameters.MemoryProvider);
+		AccessUnit->AUSize = ReadBuffer.GetLinearReadSize() + PrependData.Num();
+		AccessUnit->AUData = AccessUnit->AllocatePayloadBuffer(AccessUnit->AUSize);
+		if (PrependData.Num())
+		{
+			FMemory::Memcpy(AccessUnit->AUData, PrependData.GetData(), PrependData.Num());
+		}
+		FMemory::Memcpy(AdvancePointer(AccessUnit->AUData, PrependData.Num()), ReadBuffer.GetLinearReadData(), AccessUnit->AUSize - PrependData.Num());
+		AccessUnit->ESType = CurrentlyActiveTrackData->StreamType;
+		AccessUnit->PTS = PTS;
+		AccessUnit->PTS.SetSequenceIndex(InRequest->TimestampSequenceIndex);
+		AccessUnit->DTS = AccessUnit->PTS;
+		AccessUnit->Duration = CurrentlyActiveTrackData->DefaultDurationFromCSD;
+		AccessUnit->EarliestPTS = EarliestPTS;
+		AccessUnit->LatestPTS = LatestPTS;
+		AccessUnit->bIsFirstInSequence = CurrentlyActiveTrackData->bIsFirstInSequence;
+		AccessUnit->bIsSyncSample = true;
+		AccessUnit->bIsDummyData = false;
+		AccessUnit->AUCodecData = CurrentlyActiveTrackData->CSD;
+		AccessUnit->BufferSourceInfo = CurrentlyActiveTrackData->BufferSourceInfo;
+		AccessUnit->SequenceIndex = InRequest->TimestampSequenceIndex;
+
+		CurrentlyActiveTrackData->bIsFirstInSequence = false;
+		CurrentlyActiveTrackData->AddAccessUnit(AccessUnit);
+
+		CurrentlyActiveTrackData->DurationSuccessfullyRead = CurrentlyActiveTrackData->DefaultDurationFromCSD;
+
+		FAccessUnit::Release(AccessUnit);
+		AccessUnit = nullptr;
+	}
+	return HasReadBeenAborted() ? EHandleResult::Aborted :
+		   HasErrored() ? EHandleResult::Failed : EHandleResult::Finished;
+}
+
 void FStreamSegmentReaderCommon::FStreamHandler::InsertFillerData(TSharedPtrTS<FActiveTrackData> InActiveTrackData, TSharedPtrTS<FStreamSegmentRequestCommon> InRequest)
 {
 	if (!InActiveTrackData.IsValid() || !InRequest.IsValid())
diff --git a/Engine/Plugins/Media/ElectraPlayer/Source/ElectraPlayerRuntime/Private/Runtime/Player/StreamSegmentReaderCommon.h b/Engine/Plugins/Media/ElectraPlayer/Source/ElectraPlayerRuntime/Private/Runtime/Player/StreamSegmentReaderCommon.h
--- a/Engine/Plugins/Media/ElectraPlayer/Source/ElectraPlayerRuntime/Private/Runtime/Player/StreamSegmentReaderCommon.h
+++ b/Engine/Plugins/Media/ElectraPlayer/Source/ElectraPlayerRuntime/Private/Runtime/Player/StreamSegmentReaderCommon.h
@@ -179,7 +179,9 @@ public:
 		// ".mkv" or ".webm" container
 		Matroska_WebM,
 		// ".mp3" or ".aac" raw file with an ID3 tag
-		ID3_Raw
+		ID3_Raw,
+		// ".vtt" Raw WebVTT
+		WebVTT_Raw
 	};
 
 	// Encryption
@@ -557,6 +559,7 @@ private:
 		EHandleResult HandleTSMedia(TSharedPtrTS<FStreamSegmentRequestCommon> InRequest);
 		EHandleResult HandleID3RawMedia(TSharedPtrTS<FStreamSegmentRequestCommon> InRequest, int32 InID3HeaderSize);
 		EHandleResult HandleSideloadedMedia(TSharedPtrTS<FStreamSegmentRequestCommon> InRequest);
+		EHandleResult HandleRawSubtitleMedia(TSharedPtrTS<FStreamSegmentRequestCommon> InRequest);
 		EHandleResult HandleFillerDataSetup(TSharedPtrTS<FStreamSegmentRequestCommon> InRequest);
 		void InsertFillerData(TSharedPtrTS<FActiveTrackData> InActiveTrackData, TSharedPtrTS<FStreamSegmentRequestCommon> InRequest);
 		void CheckForInbandDASHEvents(const TSharedPtrTS<FStreamSegmentRequestCommon>& InRequest);
@@ -573,6 +576,11 @@ private:
 		int64 GetCurrentOffset() const override;
 		int64 GetTotalSize() const override;
 
+		static bool IsWebVTTHeader(const uint8* InData, int32 InSize)
+		{
+			return (InSize >= 6 && (InData[0] == 0x57 && InData[1] == 0x45 && InData[2] == 0x42 && InData[3] == 0x56 && InData[4] == 0x54 && InData[5] == 0x54)) ||
+					(InSize >= 9 && (InData[0] == 0xef && InData[1] == 0xbb && InData[2] == 0xbf && InData[3] == 0x57 && InData[4] == 0x45 && InData[5] == 0x42 && InData[6] == 0x56 && InData[7] == 0x54 && InData[8] == 0x54));
+		}
 
 		IStreamReader::CreateParam Parameters;
 		TSharedPtrTS<FStreamSegmentRequestCommon> CurrentRequest;
diff --git a/Engine/Plugins/Media/ElectraPlayer/Source/ElectraPlayerRuntime/Private/Runtime/Player/mkv/ManifestMKV.cpp b/Engine/Plugins/Media/ElectraPlayer/Source/ElectraPlayerRuntime/Private/Runtime/Player/mkv/ManifestMKV.cpp
--- a/Engine/Plugins/Media/ElectraPlayer/Source/ElectraPlayerRuntime/Private/Runtime/Player/mkv/ManifestMKV.cpp
+++ b/Engine/Plugins/Media/ElectraPlayer/Source/ElectraPlayerRuntime/Private/Runtime/Player/mkv/ManifestMKV.cpp
@@ -445,7 +445,11 @@ TSharedPtrTS<FTrackMetadata> FManifestMKVInternal::FPlayPeriodMKV::SelectMetadat
 				CandidateIndices = BCP47::FindExtendedFilteringMatch(CandList, InAttributes.Language_RFC4647.GetValue());
 				return MakeSharedTS<FTrackMetadata>(CandidateIndices.Num() ? Metadata[CandidateIndices[0]] : Metadata[0]);
 			}
-			return MakeSharedTS<FTrackMetadata>(Metadata[0]);
+			// Subtitle tracks are not selected by default. If there is no explicit selection asked for we ignore them.
+			if (StreamType != EStreamType::Subtitle)
+			{
+				return MakeSharedTS<FTrackMetadata>(Metadata[0]);
+			}
 		}
 	}
 	return nullptr;
diff --git a/Engine/Plugins/Media/ElectraPlayer/Source/ElectraPlayerRuntime/Private/Runtime/Player/mp4/ManifestMP4.cpp b/Engine/Plugins/Media/ElectraPlayer/Source/ElectraPlayerRuntime/Private/Runtime/Player/mp4/ManifestMP4.cpp
--- a/Engine/Plugins/Media/ElectraPlayer/Source/ElectraPlayerRuntime/Private/Runtime/Player/mp4/ManifestMP4.cpp
+++ b/Engine/Plugins/Media/ElectraPlayer/Source/ElectraPlayerRuntime/Private/Runtime/Player/mp4/ManifestMP4.cpp
@@ -467,7 +467,11 @@ TSharedPtrTS<FTrackMetadata> FManifestMP4Internal::FPlayPeriodMP4::SelectMetadat
 				CandidateIndices = BCP47::FindExtendedFilteringMatch(CandList, InAttributes.Language_RFC4647.GetValue());
 				return MakeSharedTS<FTrackMetadata>(CandidateIndices.Num() ? Metadata[CandidateIndices[0]] : Metadata[0]);
 			}
-			return MakeSharedTS<FTrackMetadata>(Metadata[0]);
+			// Subtitle tracks are not selected by default. If there is no explicit selection asked for we ignore them.
+			if (StreamType != EStreamType::Subtitle)
+			{
+				return MakeSharedTS<FTrackMetadata>(Metadata[0]);
+			}
 		}
 	}
 	return nullptr;
diff --git a/Engine/Plugins/Media/ElectraPlayer/Source/ElectraPlayerRuntime/Private/Runtime/StreamTypes.h b/Engine/Plugins/Media/ElectraPlayer/Source/ElectraPlayerRuntime/Private/Runtime/StreamTypes.h
--- a/Engine/Plugins/Media/ElectraPlayer/Source/ElectraPlayerRuntime/Private/Runtime/StreamTypes.h
+++ b/Engine/Plugins/Media/ElectraPlayer/Source/ElectraPlayerRuntime/Private/Runtime/StreamTypes.h
@@ -283,6 +283,17 @@ namespace Electra
 			{
 				return Top;
 			}
+
+			bool operator == (const FTranslation& rhs) const
+			{
+				return Left == rhs.Left && Top == rhs.Top;
+			}
+
+			bool operator != (const FTranslation& rhs) const
+			{
+				return !(*this == rhs);
+			}
+
 			int32 Left;
 			int32 Top;
 		};
@@ -620,7 +631,9 @@ namespace Electra
 				}
 				else if (StreamType == EStreamType::Subtitle)
 				{
-					return true;
+					return Resolution == Other.Resolution &&
+						   Translation == Other.Translation &&
+						   FrameRate == Other.FrameRate;
 				}
 			}
 			return false;
diff --git a/Engine/Plugins/Media/ElectraSubtitles/Source/ElectraSubtitles/Private/ttml/ElectraSubtitleDecoder_TTML.cpp b/Engine/Plugins/Media/ElectraSubtitles/Source/ElectraSubtitles/Private/ttml/ElectraSubtitleDecoder_TTML.cpp
--- a/Engine/Plugins/Media/ElectraSubtitles/Source/ElectraSubtitles/Private/ttml/ElectraSubtitleDecoder_TTML.cpp
+++ b/Engine/Plugins/Media/ElectraSubtitles/Source/ElectraSubtitles/Private/ttml/ElectraSubtitleDecoder_TTML.cpp
@@ -111,7 +111,7 @@ private:
 namespace ElectraSubtitleDecoderTTMLOptions
 {
 	static const FName SendEmptySubtitleDuringGaps(TEXT("sendEmptySubtitleDuringGaps"));
-	static const FName SideloadedID(TEXT("SideloadedID"));
+	static const FName SourceID(TEXT("source_id"));
 	static const FName PTO(TEXT("PTO"));
 	static const FName Width(TEXT("width"));
 	static const FName Height(TEXT("height"));
@@ -169,7 +169,7 @@ void FElectraSubtitleDecoderTTML::AddStreamedSubtitleData(const TArray<uint8>& I
 	ClearLastError();
 
 	TSharedPtr<ElectraTTMLParser::ITTMLSubtitleHandler, ESPMode::ThreadSafe> SubtitleHandler;
-	FString SideloadedID;
+	FString SourceID;
 	if (DocumentType == EDocumentType::TTMLXML)
 	{
 		// There can be only a single document active.
@@ -179,8 +179,8 @@ void FElectraSubtitleDecoderTTML::AddStreamedSubtitleData(const TArray<uint8>& I
 
 		// For sideloaded data we try to get a unique ID we can use to check if we are still operating on the
 		// same data. If we are we do not need to parse the entire document again.
-		SideloadedID = InAdditionalInfo.GetValue(ElectraSubtitleDecoderTTMLOptions::SideloadedID).SafeGetFString();
-		if (SideloadedHandler.ContentID.Equals(SideloadedID) && SideloadedHandler.Handler.IsValid())
+		SourceID = InAdditionalInfo.GetValue(ElectraSubtitleDecoderTTMLOptions::SourceID).SafeGetFString();
+		if (SideloadedHandler.ContentID.Equals(SourceID) && SideloadedHandler.Handler.IsValid())
 		{
 			SubtitleHandler = SideloadedHandler.Handler;
 
@@ -218,7 +218,7 @@ void FElectraSubtitleDecoderTTML::AddStreamedSubtitleData(const TArray<uint8>& I
 				if (DocumentType == EDocumentType::TTMLXML)
 				{
 					SideloadedHandler.Handler = SubtitleHandler;
-					SideloadedHandler.ContentID = SideloadedID;
+					SideloadedHandler.ContentID = SourceID;
 				}
 
 				FScopeLock Lock(&AccessLock);
@@ -260,6 +260,7 @@ void FElectraSubtitleDecoderTTML::Flush()
 
 	FScopeLock Lock(&AccessLock);
 	TimelineHandlers.Empty();
+	LastPlaybackUpdateAbsPos.SetToInvalid();
 }
 
 void FElectraSubtitleDecoderTTML::Start()
@@ -274,6 +275,21 @@ void FElectraSubtitleDecoderTTML::UpdatePlaybackPosition(Electra::FTimeValue InA
 {
 	TArray<FTimelineHandlers> ActiveTimelineHandlers;
 
+	// Check if the sequence index has changed due to a seek or loop.
+	// Typically we get a Flush() followed with another call here that is still prior to the loop point
+	// which results in locking the current position to the end of the looping range since we are only
+	// checking in the forward play direction.
+	// When the loop actually happens we need to reset so we start at the beginning.
+	if (!LastPlaybackUpdateAbsPos.IsValid() || LastPlaybackUpdateAbsPos.GetSequenceIndex() != InAbsolutePosition.GetSequenceIndex())
+	{
+		if (SideloadedHandler.Handler.IsValid())
+		{
+			SideloadedHandler.Handler->ClearActiveRange();
+		}
+	}
+	LastPlaybackUpdateAbsPos = InAbsolutePosition;
+
+
 	// When streaming each access unit contains (or is supposed to contain) a complete TTML document.
 	// The start and end times in the document are not adjusted for the segment duration. This means
 	// that we will also get the same subtitle more than once in every segment whose time overlaps
@@ -340,6 +356,8 @@ void FElectraSubtitleDecoderTTML::UpdatePlaybackPosition(Electra::FTimeValue InA
 				{
 					LastActiveTime = Active.AbsoluteTimeRange.End;
 				}
+				FirstActiveTime.SetSequenceIndex(InAbsolutePosition.GetSequenceIndex());
+				LastActiveTime.SetSequenceIndex(InAbsolutePosition.GetSequenceIndex());
 
 				TSharedPtr<FSubtitleDecoderOutputTTML, ESPMode::ThreadSafe> Out = MakeShared<FSubtitleDecoderOutputTTML, ESPMode::ThreadSafe>();
 				Out->SetText(Text);
diff --git a/Engine/Plugins/Media/ElectraSubtitles/Source/ElectraSubtitles/Private/ttml/ElectraSubtitleDecoder_TTML.h b/Engine/Plugins/Media/ElectraSubtitles/Source/ElectraSubtitles/Private/ttml/ElectraSubtitleDecoder_TTML.h
--- a/Engine/Plugins/Media/ElectraSubtitles/Source/ElectraSubtitles/Private/ttml/ElectraSubtitleDecoder_TTML.h
+++ b/Engine/Plugins/Media/ElectraSubtitles/Source/ElectraSubtitles/Private/ttml/ElectraSubtitleDecoder_TTML.h
@@ -70,6 +70,7 @@ private:
 	FCriticalSection AccessLock;
 	TArray<FTimelineHandlers> TimelineHandlers;
 	FSideloadedSubtitleHandler SideloadedHandler;
+	Electra::FTimeValue LastPlaybackUpdateAbsPos;
 
 	int32				Width = 0;
 	int32				Height = 0;
diff --git a/Engine/Plugins/Media/ElectraSubtitles/Source/ElectraSubtitles/Private/wvtt/ElectraSubtitleDecoder_WVTT.cpp b/Engine/Plugins/Media/ElectraSubtitles/Source/ElectraSubtitles/Private/wvtt/ElectraSubtitleDecoder_WVTT.cpp
--- a/Engine/Plugins/Media/ElectraSubtitles/Source/ElectraSubtitles/Private/wvtt/ElectraSubtitleDecoder_WVTT.cpp
+++ b/Engine/Plugins/Media/ElectraSubtitles/Source/ElectraSubtitles/Private/wvtt/ElectraSubtitleDecoder_WVTT.cpp
@@ -1,6 +1,7 @@
 // Copyright Epic Games, Inc. All Rights Reserved.
 
 #include "wvtt/ElectraSubtitleDecoder_WVTT.h"
+#include "wvtt/WebVTTParser.h"
 #include "ElectraSubtitleDecoderFactory.h"
 #include "ElectraSubtitleModule.h"
 #include "MediaDecoderOutput.h"
@@ -15,6 +16,8 @@ namespace OptionKeys
 	static const FName OffsetX(TEXT("offset_x"));
 	static const FName OffsetY(TEXT("offset_y"));
 	static const FName Timescale(TEXT("timescale"));
+	static const FName SourceID(TEXT("source_id"));
+	static const FName SendEmptySubtitleDuringGaps(TEXT("sendEmptySubtitleDuringGaps"));
 }
 
 #if !PLATFORM_LITTLE_ENDIAN
@@ -155,7 +158,11 @@ public:
 void FElectraSubtitleDecoderWVTT::RegisterCodecs(IElectraSubtitleDecoderFactoryRegistry& InRegistry)
 {
 	static FElectraSubtitleDecoderFactoryWVTT Factory;
-	TArray<IElectraSubtitleDecoderFactoryRegistry::FCodecInfo> CodecInfos { { FString(TEXT("wvtt")), 0} };
+	TArray<IElectraSubtitleDecoderFactoryRegistry::FCodecInfo> CodecInfos
+	{
+		{ FString(TEXT("wvtt")), 0},			// codec
+		{ FString(TEXT("text/vtt")), 0}			// mimetype
+	};
 	InRegistry.AddDecoderFactory(CodecInfos, &Factory);
 }
 
@@ -193,12 +200,12 @@ public:
 		Timestamp.SequenceIndex = InTimestamp.GetSequenceIndex();
 	}
 
-	void SetID(const FString& InID)
+	void SetID(const FString& InID, bool bInIsGeneratedID)
 	{
 		ID = InID;
+		bIsGeneratedID = bInIsGeneratedID;
 	}
 
-
 	const TArray<uint8>& GetData() override
 	{
 		return TextAsArray;
@@ -228,11 +235,24 @@ public:
 	{
 		return ID;
 	}
+
+	bool operator == (const FSubtitleDecoderOutputWVTT& Other) const
+	{
+		return Timestamp.Time == Other.Timestamp.Time
+			&& (bIsGeneratedID || Other.bIsGeneratedID || ID.Equals(Other.ID))
+			&& Duration == Other.Duration && TextAsArray == Other.TextAsArray;
+	}
+	
+	bool operator != (const FSubtitleDecoderOutputWVTT& Other) const
+	{
+		return !(this->operator==(Other));
+	}
 private:
 	TArray<uint8> TextAsArray;
 	FString ID;
 	FDecoderTimeStamp Timestamp;
 	FTimespan Duration;
+	bool bIsGeneratedID = false;
 };
 
 
@@ -248,70 +268,78 @@ FElectraSubtitleDecoderWVTT::~FElectraSubtitleDecoderWVTT()
 
 bool FElectraSubtitleDecoderWVTT::InitializeStreamWithCSD(const TArray<uint8>& InCSD, const Electra::FParamDict& InAdditionalInfo)
 {
-	#define RETURN_IF_ERROR(expr)								\
-		if (!expr)												\
-		{														\
-			UE_LOG(LogElectraSubtitles, Error, TEXT("Not enough CSD bytes to parse WVTT text sample entry"));	\
-			return false;										\
-		}														\
-
-	// Get dimension, placement offset and other data from the sideband dictionary.
-	Width = (int32) InAdditionalInfo.GetValue(ElectraSubtitleDecoderWVTTUtils::OptionKeys::Width).SafeGetInt64(0);
-	Height = (int32) InAdditionalInfo.GetValue(ElectraSubtitleDecoderWVTTUtils::OptionKeys::Height).SafeGetInt64(0);
-	TranslationX = (int32) InAdditionalInfo.GetValue(ElectraSubtitleDecoderWVTTUtils::OptionKeys::OffsetX).SafeGetInt64(0);
-	TranslationY = (int32) InAdditionalInfo.GetValue(ElectraSubtitleDecoderWVTTUtils::OptionKeys::OffsetY).SafeGetInt64(0);
-	Timescale = (uint32) InAdditionalInfo.GetValue(ElectraSubtitleDecoderWVTTUtils::OptionKeys::Timescale).SafeGetInt64(0);
-
-	ElectraSubtitleDecoderWVTTUtils::FDataReaderMP4 r(InCSD);
-
-	// The CSD is a WVTTSampleEntry of an mp4 (ISO/IEC 14496-12) file and interpreted as per
-	// ISO/IEC 14496-30 Section 7.5 Sample entry format
-	RETURN_IF_ERROR(r.ReadBytes(nullptr, 6));
-	RETURN_IF_ERROR(r.Read(DataReferenceIndex));
-
-	// Read the configuration boxes.
-	// There needs to be a 'vttC' box.
-	// 'vlab' and 'btrt' boxes are optional. Anything else is ignored.
-	while(r.GetNumBytesRemaining() > 0)
+	// Get data from the sideband dictionary.
+	bool bIsSideloaded = InCSD.IsEmpty();
+	// Check if the CSD starts with "WEBVTT" or BOM+"WEBVTT"
+	bool bIsRaw = (InCSD.Num() >= 6 && InCSD[0] == 0x57 && InCSD[1] == 0x45 && InCSD[2] == 0x42 && InCSD[3] == 0x56 && InCSD[4] == 0x54 && InCSD[5] == 0x54) ||
+				  (InCSD.Num() >= 9 && InCSD[0] == 0xef && InCSD[1] == 0xbb && InCSD[2] == 0xbf && InCSD[3] == 0x57 && InCSD[4] == 0x45 && InCSD[5] == 0x42 && InCSD[6] == 0x56 && InCSD[7] == 0x54 && InCSD[8] == 0x54);
+	if (bIsSideloaded || bIsRaw)
 	{
-		uint32 BoxLen, BoxType;
-		if (!r.Read(BoxLen) || !r.Read(BoxType))
-		{
-			UE_LOG(LogElectraSubtitles, Error, TEXT("Bad WVTT box in CSD, ignoring."));
-			break;
-		}
-		if (BoxLen < 8 || (int32)(BoxLen-8 )> r.GetNumBytesRemaining())
-		{
-			UE_LOG(LogElectraSubtitles, Error, TEXT("Bad WVTT box in CSD, ignoring."));
-			break;
-		}
-
-		if (BoxType == ElectraSubtitleDecoderWVTTUtils::FDataReaderMP4::BoxType_vttC)
-		{
-			RETURN_IF_ERROR(r.ReadString(Configuration, BoxLen - 8));
-		}
-		else if (BoxType == ElectraSubtitleDecoderWVTTUtils::FDataReaderMP4::BoxType_vlab)
-		{
-			RETURN_IF_ERROR(r.ReadString(Label, BoxLen - 8));
-		}
-		else
+		bNeedsParsing = true;
+	}
+	else
+	{
+		#define RETURN_IF_ERROR(expr)																				\
+			if (!expr)																								\
+			{																										\
+				UE_LOG(LogElectraSubtitles, Error, TEXT("Not enough CSD bytes to parse WVTT text sample entry"));	\
+				return false;																						\
+			}
+		ElectraSubtitleDecoderWVTTUtils::FDataReaderMP4 r(InCSD);
+		FString Configuration;
+		FString Label;
+
+		// The CSD is a WVTTSampleEntry of an mp4 (ISO/IEC 14496-12) file and interpreted as per
+		// ISO/IEC 14496-30 Section 7.5 Sample entry format
+		RETURN_IF_ERROR(r.ReadBytes(nullptr, 6));
+		uint16 DataReferenceIndex = 0;
+		RETURN_IF_ERROR(r.Read(DataReferenceIndex));
+
+		// Read the configuration boxes.
+		// There needs to be a 'vttC' box.
+		// 'vlab' and 'btrt' boxes are optional. Anything else is ignored.
+		while(r.GetNumBytesRemaining() > 0)
 		{
-			// Note: We ignore a potential 'btrt' box here. There is no use for us.
-			if (!r.ReadBytes(nullptr, BoxLen - 8))
+			uint32 BoxLen, BoxType;
+			if (!r.Read(BoxLen) || !r.Read(BoxType))
 			{
+				UE_LOG(LogElectraSubtitles, Error, TEXT("Bad WVTT box in CSD, ignoring."));
 				break;
 			}
+			if (BoxLen < 8 || (int32)(BoxLen-8 )> r.GetNumBytesRemaining())
+			{
+				UE_LOG(LogElectraSubtitles, Error, TEXT("Bad WVTT box in CSD, ignoring."));
+				break;
+			}
+
+			if (BoxType == ElectraSubtitleDecoderWVTTUtils::FDataReaderMP4::BoxType_vttC)
+			{
+				RETURN_IF_ERROR(r.ReadString(Configuration, BoxLen - 8));
+			}
+			else if (BoxType == ElectraSubtitleDecoderWVTTUtils::FDataReaderMP4::BoxType_vlab)
+			{
+				RETURN_IF_ERROR(r.ReadString(Label, BoxLen - 8));
+			}
+			else
+			{
+				// Note: We ignore a potential 'btrt' box here. There is no use for us.
+				if (!r.ReadBytes(nullptr, BoxLen - 8))
+				{
+					break;
+				}
+			}
 		}
-	}
 
-	// The "configuration" is everything up to the first cue, so at least the string "WEBVTT".
-	if (Configuration.IsEmpty() || !Configuration.StartsWith(TEXT("WEBVTT")))
-	{
-		UE_LOG(LogElectraSubtitles, Error, TEXT("Bad WVTT configuration box 'vttC'!"));
-		return false;
+		// The "configuration" is everything up to the first cue, so at least the string "WEBVTT".
+		if (Configuration.IsEmpty() || !Configuration.StartsWith(TEXT("WEBVTT")))
+		{
+			UE_LOG(LogElectraSubtitles, Error, TEXT("Bad WVTT configuration box 'vttC'!"));
+			return false;
+		}
+		#undef RETURN_IF_ERROR
 	}
 
-	#undef RETURN_IF_ERROR
+	bSendEmptySubtitleDuringGaps = InAdditionalInfo.GetValue(ElectraSubtitleDecoderWVTTUtils::OptionKeys::SendEmptySubtitleDuringGaps).SafeGetBool();
 	return true;
 }
 
@@ -327,101 +355,163 @@ Electra::FTimeValue FElectraSubtitleDecoderWVTT::GetStreamedDeliveryTimeOffset()
 
 void FElectraSubtitleDecoderWVTT::AddStreamedSubtitleData(const TArray<uint8>& InData, Electra::FTimeValue InAbsoluteTimestamp, Electra::FTimeValue InDuration, const Electra::FParamDict& InAdditionalInfo)
 {
-	#define RETURN_IF_ERROR(expr)																\
-		if (!expr)																				\
-		{																						\
-			UE_LOG(LogElectraSubtitles, Error, TEXT("Bad WVTT text sample box, ignoring."));	\
-			return;																				\
-		}																						\
-
-	ElectraSubtitleDecoderWVTTUtils::FDataReaderMP4 r(InData);
-
-	// List of collected subtitles.
-	TArray<FSubtitleEntry> Subtitles;
-	// The currently worked-on subtitle.
-	FSubtitleEntry Sub;
-
-	bool bInsideCue = false;
-	while(r.GetNumBytesRemaining() > 0)
+	if (bNeedsParsing)
 	{
-		uint32 BoxLen, BoxType;
-		if (!r.Read(BoxLen) || !r.Read(BoxType))
+		// Check if we already have this.
+		FScopeLock lock(&ParsedTimerangeLock);
+		for(int32 i=0; i<ParsedTimeranges.Num(); ++i)
 		{
-			UE_LOG(LogElectraSubtitles, Error, TEXT("Bad WVTT text sample box, ignoring."));
-			break;
+			if (ParsedTimeranges[i]->SourceID.Equals(InAdditionalInfo.GetValue(ElectraSubtitleDecoderWVTTUtils::OptionKeys::SourceID).SafeGetFString()) &&
+				ParsedTimeranges[i]->AbsoluteStartTime == InAbsoluteTimestamp &&
+				ParsedTimeranges[i]->Duration == InDuration)
+			{
+				// Update the sequence index in case of looping. Otherwise the range will get removed when
+				// the playback position is outside its range.
+				ParsedTimeranges[i]->AbsoluteStartTime.SetSequenceIndex(InAbsoluteTimestamp.GetSequenceIndex());
+				return;
+			}
 		}
-		if (BoxLen < 8 || (int32)(BoxLen-8 )> r.GetNumBytesRemaining())
+		ParsedTimerangeLock.Unlock();
+		
+		TSharedPtr<FParsedTimerange, ESPMode::ThreadSafe> ParsedTimerange = MakeShared<FParsedTimerange, ESPMode::ThreadSafe>();
+		ParsedTimerange->Parser = ElectraWebVTTParser::IWebVTTParser::Create();
+		ParsedTimerange->AbsoluteStartTime = InAbsoluteTimestamp;
+		ParsedTimerange->Duration = InDuration;
+		ParsedTimerange->SourceID = InAdditionalInfo.GetValue(ElectraSubtitleDecoderWVTTUtils::OptionKeys::SourceID).SafeGetFString();
+		Electra::FParamDict noOpts;
+		if (ParsedTimerange->Parser->ParseWebVTT(InData, noOpts))
 		{
-			UE_LOG(LogElectraSubtitles, Error, TEXT("Bad WVTT text sample box, ignoring."));
-			break;
-		}
+			FScopeLock lock2(&ParsedTimerangeLock);
+			ParsedTimeranges.Emplace(MoveTemp(ParsedTimerange));
+			ParsedTimeranges.Sort([](const TSharedPtr<FParsedTimerange, ESPMode::ThreadSafe>& a, const TSharedPtr<FParsedTimerange, ESPMode::ThreadSafe>& b){ return a->AbsoluteStartTime < b->AbsoluteStartTime;});
 
-		// An empty cue?
-		if (BoxType == ElectraSubtitleDecoderWVTTUtils::FDataReaderMP4::BoxType_vtte)
-		{
-			if (bInsideCue)
-			{
-				Subtitles.Emplace(MoveTemp(Sub));
-			}
-			bInsideCue = true;
-			// This must be the only entry, so we stop parsing now.
-			// If there are additional boxes this is an authoring error we ignore.
-			break;
+			// A change in parsed time ranges requires an immediate re-evaluation.
+			NextEvaluationAt = FTimespan::MinValue();
 		}
-		// An additional text box (a comment)?
-		else if (BoxType == ElectraSubtitleDecoderWVTTUtils::FDataReaderMP4::BoxType_vtta)
+		else
 		{
-			if (bInsideCue)
-			{
-				Subtitles.Emplace(MoveTemp(Sub));
-			}
-			bInsideCue = true;
-			Sub.bIsAdditionalCue = true;
-			RETURN_IF_ERROR(r.ReadString(Sub.Text, BoxLen - 8));
+			UE_LOG(LogElectraSubtitles, Error, TEXT("Bad WebVTT document, ignoring."));
+			return;
 		}
-		// A cue?
-		else if (BoxType == ElectraSubtitleDecoderWVTTUtils::FDataReaderMP4::BoxType_vttc)
+	}
+	else
+	{
+		#define RETURN_IF_ERROR(expr)																\
+			if (!expr)																				\
+			{																						\
+				UE_LOG(LogElectraSubtitles, Error, TEXT("Bad WVTT text sample box, ignoring."));	\
+				return;																				\
+			}
+
+		ElectraSubtitleDecoderWVTTUtils::FDataReaderMP4 r(InData);
+
+		struct FSubtitleEntry
 		{
-			if (bInsideCue)
+			FString Text;
+			TOptional<int32> SourceID;
+			TOptional<FString> CurrentTime;
+			TOptional<FString> ID;
+			TOptional<FString> Settings;
+			bool bIsAdditionalCue = false;
+		};
+
+		// List of collected subtitles.
+		TArray<FSubtitleEntry> Subtitles;
+		// The currently worked-on subtitle.
+		FSubtitleEntry Sub;
+
+		bool bInsideCue = false;
+		while(r.GetNumBytesRemaining() > 0)
+		{
+			uint32 BoxLen, BoxType;
+			if (!r.Read(BoxLen) || !r.Read(BoxType))
 			{
-				Subtitles.Emplace(MoveTemp(Sub));
+				UE_LOG(LogElectraSubtitles, Error, TEXT("Bad WVTT text sample box, ignoring."));
+				break;
 			}
-			bInsideCue = true;
-		}
-		else if (bInsideCue)
-		{
-			// Cue source ID?
-			if (BoxType == ElectraSubtitleDecoderWVTTUtils::FDataReaderMP4::BoxType_vsid)
+			if (BoxLen < 8 || (int32)(BoxLen-8 )> r.GetNumBytesRemaining())
 			{
-				int32 source_ID;
-				RETURN_IF_ERROR(r.Read(source_ID));
-				Sub.SourceID = source_ID;
+				UE_LOG(LogElectraSubtitles, Error, TEXT("Bad WVTT text sample box, ignoring."));
+				break;
 			}
-			// Cue time?
-			else if (BoxType == ElectraSubtitleDecoderWVTTUtils::FDataReaderMP4::BoxType_ctim)
+
+			// An empty cue?
+			if (BoxType == ElectraSubtitleDecoderWVTTUtils::FDataReaderMP4::BoxType_vtte)
 			{
-				FString s;
-				RETURN_IF_ERROR(r.ReadString(s, BoxLen - 8));
-				Sub.CurrentTime = s;
+				if (bInsideCue)
+				{
+					Subtitles.Emplace(MoveTemp(Sub));
+				}
+				bInsideCue = true;
+				// This must be the only entry, so we stop parsing now.
+				// If there are additional boxes this is an authoring error we ignore.
+				break;
 			}
-			// ID?
-			else if (BoxType == ElectraSubtitleDecoderWVTTUtils::FDataReaderMP4::BoxType_iden)
+			// An additional text box (a comment)?
+			else if (BoxType == ElectraSubtitleDecoderWVTTUtils::FDataReaderMP4::BoxType_vtta)
 			{
-				FString s;
-				RETURN_IF_ERROR(r.ReadString(s, BoxLen - 8));
-				Sub.ID = s;
+				if (bInsideCue)
+				{
+					Subtitles.Emplace(MoveTemp(Sub));
+				}
+				bInsideCue = true;
+				Sub.bIsAdditionalCue = true;
+				RETURN_IF_ERROR(r.ReadString(Sub.Text, BoxLen - 8));
 			}
-			// Settings?
-			else if (BoxType == ElectraSubtitleDecoderWVTTUtils::FDataReaderMP4::BoxType_sttg)
+			// A cue?
+			else if (BoxType == ElectraSubtitleDecoderWVTTUtils::FDataReaderMP4::BoxType_vttc)
 			{
-				FString s;
-				RETURN_IF_ERROR(r.ReadString(s, BoxLen - 8));
-				Sub.Settings = s;
+				if (bInsideCue)
+				{
+					Subtitles.Emplace(MoveTemp(Sub));
+				}
+				bInsideCue = true;
 			}
-			// Payload?
-			else if (BoxType == ElectraSubtitleDecoderWVTTUtils::FDataReaderMP4::BoxType_payl)
+			else if (bInsideCue)
 			{
-				RETURN_IF_ERROR(r.ReadString(Sub.Text, BoxLen - 8));
+				// Cue source ID?
+				if (BoxType == ElectraSubtitleDecoderWVTTUtils::FDataReaderMP4::BoxType_vsid)
+				{
+					int32 source_ID;
+					RETURN_IF_ERROR(r.Read(source_ID));
+					Sub.SourceID = source_ID;
+				}
+				// Cue time?
+				else if (BoxType == ElectraSubtitleDecoderWVTTUtils::FDataReaderMP4::BoxType_ctim)
+				{
+					FString s;
+					RETURN_IF_ERROR(r.ReadString(s, BoxLen - 8));
+					Sub.CurrentTime = s;
+				}
+				// ID?
+				else if (BoxType == ElectraSubtitleDecoderWVTTUtils::FDataReaderMP4::BoxType_iden)
+				{
+					FString s;
+					RETURN_IF_ERROR(r.ReadString(s, BoxLen - 8));
+					Sub.ID = s;
+				}
+				// Settings?
+				else if (BoxType == ElectraSubtitleDecoderWVTTUtils::FDataReaderMP4::BoxType_sttg)
+				{
+					FString s;
+					RETURN_IF_ERROR(r.ReadString(s, BoxLen - 8));
+					Sub.Settings = s;
+				}
+				// Payload?
+				else if (BoxType == ElectraSubtitleDecoderWVTTUtils::FDataReaderMP4::BoxType_payl)
+				{
+					RETURN_IF_ERROR(r.ReadString(Sub.Text, BoxLen - 8));
+					// Strip any trailing newlines. They should not have been added by the muxing tool.
+					while(!Sub.Text.IsEmpty() && (Sub.Text[Sub.Text.Len()-1] == TCHAR('\n') || Sub.Text[Sub.Text.Len()-1] == TCHAR('\r')))
+					{
+						Sub.Text.LeftChopInline(1);
+					}
+				}
+				// Something else.
+				else if (!r.ReadBytes(nullptr, BoxLen - 8))
+				{
+					break;
+				}
 			}
 			// Something else.
 			else if (!r.ReadBytes(nullptr, BoxLen - 8))
@@ -429,40 +519,35 @@ void FElectraSubtitleDecoderWVTT::AddStreamedSubtitleData(const TArray<uint8>& I
 				break;
 			}
 		}
-		// Something else.
-		else if (!r.ReadBytes(nullptr, BoxLen - 8))
+		// At the end of the data add the currently worked on cue, if there is one.
+		if (bInsideCue)
 		{
-			break;
+			Subtitles.Emplace(MoveTemp(Sub));
 		}
-	}
-	// At the end of the data add the currently worked on cue, if there is one.
-	if (bInsideCue)
-	{
-		Subtitles.Emplace(MoveTemp(Sub));
-	}
 
-	for(auto &Cue : Subtitles)
-	{
-		if (!Cue.bIsAdditionalCue)
+		for(auto &Cue : Subtitles)
 		{
-			TSharedPtr<FSubtitleDecoderOutputWVTT, ESPMode::ThreadSafe> Out = MakeShared<FSubtitleDecoderOutputWVTT, ESPMode::ThreadSafe>();
-			Out->SetTimestamp(InAbsoluteTimestamp);
-			Out->SetDuration(InDuration);
-			// This decoder returns plain text only. Remove all formatting tags.
-			Out->SetText(RemoveAllTags(Cue.Text));
-			if (Cue.ID.IsSet())
-			{
-				Out->SetID(Cue.ID.GetValue());
-			}
-			else
+			if (!Cue.bIsAdditionalCue)
 			{
-				Out->SetID(FString(TEXT("<")) + LexToString(++NextID) + FString(TEXT(">")));
+				TSharedPtr<FSubtitleDecoderOutputWVTT, ESPMode::ThreadSafe> Out = MakeShared<FSubtitleDecoderOutputWVTT, ESPMode::ThreadSafe>();
+				Out->SetTimestamp(InAbsoluteTimestamp);
+				Out->SetDuration(InDuration);
+				// This decoder returns plain text only. Remove all formatting tags.
+				Out->SetText(ElectraWebVTTParser::IWebVTTParser::ProcessCueTextForType(Cue.Text, ElectraWebVTTParser::IWebVTTParser::EWebVTTType::Subtitle));
+				if (Cue.ID.IsSet())
+				{
+					Out->SetID(Cue.ID.GetValue(), false);
+				}
+				else
+				{
+					Out->SetID(FString::Printf(TEXT("<%u>"), ++NextID), true);
+				}
+				ParsedSubtitleDelegate.Broadcast(Out);
 			}
-			ParsedSubtitleDelegate.Broadcast(Out);
 		}
-	}
 
-	#undef RETURN_IF_ERROR
+		#undef RETURN_IF_ERROR
+	}
 }
 
 void FElectraSubtitleDecoderWVTT::SignalStreamedSubtitleEOD()
@@ -471,6 +556,11 @@ void FElectraSubtitleDecoderWVTT::SignalStreamedSubtitleEOD()
 
 void FElectraSubtitleDecoderWVTT::Flush()
 {
+	FScopeLock lock(&ParsedTimerangeLock);
+	NextEvaluationAt = FTimespan::MinValue();
+	LastSentSubtitle.Reset();
+	LastPlaybackUpdateAbsPos.SetToInvalid();
+	SendEmptySubtitleAt = FTimespan::MinValue();
 }
 
 void FElectraSubtitleDecoderWVTT::Start()
@@ -483,46 +573,124 @@ void FElectraSubtitleDecoderWVTT::Stop()
 
 void FElectraSubtitleDecoderWVTT::UpdatePlaybackPosition(Electra::FTimeValue InAbsolutePosition, Electra::FTimeValue InLocalPosition)
 {
-}
+	ParsedTimerangeLock.Lock();
+	TArray<TSharedPtr<FParsedTimerange, ESPMode::ThreadSafe>> TimeRanges;
+	for(int32 i=0; i<ParsedTimeranges.Num(); ++i)
+	{
+		if (InAbsolutePosition > ParsedTimeranges[i]->AbsoluteStartTime + ParsedTimeranges[i]->Duration && 
+			InAbsolutePosition.GetSequenceIndex() >= ParsedTimeranges[i]->AbsoluteStartTime.GetSequenceIndex())
+		{
+			ParsedTimeranges.RemoveAt(i);
+			--i;
+		}
+		else
+		{
+			TimeRanges.Emplace(ParsedTimeranges[i]);
+		}
+	}
+
+	// Check if the sequence index has changed due to a seek or loop.
+	// Typically we get a Flush() followed with another call here that is still prior to the loop point
+	// which results in locking the current position to the end of the looping range since we are only
+	// checking in the forward play direction.
+	// When the loop actually happens we need to reset so we start at the beginning.
+	if (!LastPlaybackUpdateAbsPos.IsValid() || LastPlaybackUpdateAbsPos.GetSequenceIndex() != InAbsolutePosition.GetSequenceIndex())
+	{
+		NextEvaluationAt = FTimespan::MinValue();
+		SendEmptySubtitleAt = FTimespan::MinValue();
+		for(int32 i=0; i<ParsedTimeranges.Num(); ++i)
+		{
+			ParsedTimeranges[i]->CurrentCueIterator.Reset();
+		}
+	}
+	LastPlaybackUpdateAbsPos = InAbsolutePosition;
 
+	FTimespan NextEval = NextEvaluationAt;
+	bool bWasForcedEval = NextEvaluationAt == FTimespan::MinValue();
+	ParsedTimerangeLock.Unlock();
 
+	if (InAbsolutePosition.GetAsTimespan() < NextEval)
+	{
+		return;
+	}
 
-FString FElectraSubtitleDecoderWVTT::RemoveAllTags(const FString& InString)
-{
-	/*
-		For simplicity we do not attempt to parse the tags.
-		Instead we remove everything enclosed in a '<' '>' pair, including the less and greater characters themselves.
-		We also do not check for balanced <>. A sequence like <<x> will also be removed entirely.
-
-		See: https://developer.mozilla.org/en-US/docs/Web/API/WebVTT_API
-	*/
-	FString Result;
-	const TCHAR * const TagStart = TEXT("<");
-	const TCHAR * const TagEnd = TEXT(">");
-	int32 CurrentPos = 0;
-
-	while(CurrentPos < InString.Len())
+	// Get the current cues
+	TArray<TUniquePtr<ElectraWebVTTParser::IWebVTTParser::ICue>> Cues;
+	FTimespan ChangesAt;
+	NextEval = FTimespan::MaxValue();
+	for(int32 i=0; i<TimeRanges.Num(); ++i)
+	{
+		TimeRanges[i]->Parser->GetCuesAtTime(Cues, ChangesAt, TimeRanges[i]->CurrentCueIterator, InAbsolutePosition.GetAsTimespan());
+		NextEval = FMath::Min(NextEval, ChangesAt);
+	}
+	ParsedTimerangeLock.Lock();
+	NextEvaluationAt = NextEval;
+	ParsedTimerangeLock.Unlock();
+	if (Cues.Num())
 	{
-		// Find the next opening tag
-		int32 ltPos = InString.Find(TagStart, ESearchCase::CaseSensitive, ESearchDir::FromStart, CurrentPos);
-		int32 gtPos = INDEX_NONE;
-		if (ltPos != INDEX_NONE)
+		FString CombinedPlainText;
+		FTimespan CueStartTime = FTimespan::MaxValue();
+		FTimespan CueEndTime = FTimespan::MaxValue();
+		for(int32 i=0; i<Cues.Num(); ++i)
+		{
+			CueStartTime = FMath::Min(CueStartTime, Cues[i]->GetStartTime());
+			CueEndTime = FMath::Min(CueEndTime, Cues[i]->GetEndTime());
+			if (i > 0)
+			{
+				CombinedPlainText.Append(TEXT("\n"));
+			}
+
+			CombinedPlainText.Append(ElectraWebVTTParser::IWebVTTParser::ProcessCueTextForType(Cues[i]->GetText(), ElectraWebVTTParser::IWebVTTParser::EWebVTTType::Subtitle));
+		}
+		// Get the cue ID only when there is a single cue. If we had to aggregate cues then there is no unique ID.
+		FString CueID = Cues.Num() == 1 ? Cues[0]->GetID() : FString();
+
+		TSharedPtr<FSubtitleDecoderOutputWVTT, ESPMode::ThreadSafe> Out = MakeShared<FSubtitleDecoderOutputWVTT, ESPMode::ThreadSafe>();
+		Out->SetText(CombinedPlainText);
+		Electra::FTimeValue tv;
+		tv.SetFromTimespan(CueStartTime);
+		tv.SetSequenceIndex(InAbsolutePosition.GetSequenceIndex());
+		Out->SetTimestamp(tv);
+		Electra::FTimeValue dur;
+		dur.SetFromTimespan(CueEndTime - CueStartTime);
+		Out->SetDuration(dur);
+		Out->SetID(CueID.Len() ? CueID : FString::Printf(TEXT("<%u>"), ++NextID), CueID.IsEmpty());
+
+		// Do not send the same subtitle again. This happens when adding a new streamed subtitle had to invalidate the
+		// time at which we had to update.
+		if (bWasForcedEval && LastSentSubtitle.IsValid() && *LastSentSubtitle == *Out)
 		{
-			// Find tag close char. If there is none, we do not remove the opening character.
-			gtPos = InString.Find(TagEnd, ESearchCase::CaseSensitive, ESearchDir::FromStart, ltPos);
-			ltPos = gtPos == INDEX_NONE ? gtPos : ltPos;
+			Out.Reset();
+		}
+		if (Out.IsValid())
+		{
+			ParsedSubtitleDelegate.Broadcast(Out);
+			LastSentSubtitle = Out;
 		}
-		Result += InString.Mid(CurrentPos, ltPos == INDEX_NONE ? MAX_int32 : ltPos - CurrentPos);
-		CurrentPos = gtPos != INDEX_NONE ? gtPos+1 : ltPos == INDEX_NONE ? MAX_int32 : ltPos+1;
-	}
 
-	// Replace common XML entities.
-	Result.ReplaceInline(TEXT("&lt;"),   TEXT("<"),  ESearchCase::CaseSensitive);
-	Result.ReplaceInline(TEXT("&gt;"),   TEXT(">"),  ESearchCase::CaseSensitive);
-	Result.ReplaceInline(TEXT("&nbsp;"), TEXT(" "),  ESearchCase::CaseSensitive);
-	Result.ReplaceInline(TEXT("&amp;"),  TEXT("&"),  ESearchCase::CaseSensitive);
-	Result.ReplaceInline(TEXT("&quot;"), TEXT("\""), ESearchCase::CaseSensitive);
+		if (bSendEmptySubtitleDuringGaps)
+		{
+			SendEmptySubtitleAt = CueEndTime;
+		}
+	}
+	else if (bSendEmptySubtitleDuringGaps && SendEmptySubtitleAt != FTimespan::MinValue())
+	{
+		TSharedPtr<FSubtitleDecoderOutputWVTT, ESPMode::ThreadSafe> Out = MakeShared<FSubtitleDecoderOutputWVTT, ESPMode::ThreadSafe>();
+		Electra::FTimeValue tv;
+		tv.SetFromTimespan(SendEmptySubtitleAt);
+		tv.SetSequenceIndex(InAbsolutePosition.GetSequenceIndex());
+		Out->SetTimestamp(tv);
+		FTimespan d(NextEval - SendEmptySubtitleAt);
+		if (d < FTimespan::Zero() || d.GetTicks() > ETimespan::TicksPerSecond)
+		{
+			d = FTimespan(ETimespan::TicksPerSecond);
+		}
+		Electra::FTimeValue dur;
+		dur.SetFromTimespan(d);
+		Out->SetDuration(dur);
+		Out->SetID(FString::Printf(TEXT("<%u>"), ++NextID), true);
+		ParsedSubtitleDelegate.Broadcast(Out);
 
-	return Result;
+		SendEmptySubtitleAt = FTimespan::MinValue();
+	}
 }
-
diff --git a/Engine/Plugins/Media/ElectraSubtitles/Source/ElectraSubtitles/Private/wvtt/ElectraSubtitleDecoder_WVTT.h b/Engine/Plugins/Media/ElectraSubtitles/Source/ElectraSubtitles/Private/wvtt/ElectraSubtitleDecoder_WVTT.h
--- a/Engine/Plugins/Media/ElectraSubtitles/Source/ElectraSubtitles/Private/wvtt/ElectraSubtitleDecoder_WVTT.h
+++ b/Engine/Plugins/Media/ElectraSubtitles/Source/ElectraSubtitles/Private/wvtt/ElectraSubtitleDecoder_WVTT.h
@@ -3,8 +3,12 @@
 #pragma once
 
 #include "IElectraSubtitleDecoder.h"
+#include "wvtt/WebVTTParser.h"
 
 class IElectraSubtitleDecoderFactoryRegistry;
+class FSubtitleDecoderOutputWVTT;
+
+
 
 /**
  * WebVTT subtitle decoder (https://www.w3.org/TR/webvtt1/)
@@ -32,28 +36,24 @@ public:
 	virtual void UpdatePlaybackPosition(Electra::FTimeValue InAbsolutePosition, Electra::FTimeValue InLocalPosition) override;
 
 private:
-	struct FSubtitleEntry
+
+	struct FParsedTimerange
 	{
-		FString Text;
-		TOptional<int32> SourceID;
-		TOptional<FString> CurrentTime;
-		TOptional<FString> ID;
-		TOptional<FString> Settings;
-		bool bIsAdditionalCue = false;
+		TSharedPtr<ElectraWebVTTParser::IWebVTTParser, ESPMode::ThreadSafe> Parser;
+		TUniquePtr<ElectraWebVTTParser::IWebVTTParser::ICueIterator> CurrentCueIterator;
+		Electra::FTimeValue AbsoluteStartTime;
+		Electra::FTimeValue Duration;
+		FString SourceID;
 	};
-
-	FString RemoveAllTags(const FString& InString);
+	FCriticalSection ParsedTimerangeLock;
+	TArray<TSharedPtr<FParsedTimerange, ESPMode::ThreadSafe>> ParsedTimeranges;
+	FTimespan NextEvaluationAt { FTimespan::MinValue() };
+	TSharedPtr<FSubtitleDecoderOutputWVTT, ESPMode::ThreadSafe> LastSentSubtitle;
+	Electra::FTimeValue LastPlaybackUpdateAbsPos;
 
 	FOnSubtitleReceivedDelegate ParsedSubtitleDelegate;
 	uint32 NextID = 0;
-
-	uint16				DataReferenceIndex = 0;
-	int32				Width = 0;
-	int32				Height = 0;
-	int32				TranslationX = 0;
-	int32				TranslationY = 0;
-	uint32				Timescale = 0;
-
-	FString				Configuration;
-	FString				Label;
+	bool bNeedsParsing = false;
+	bool bSendEmptySubtitleDuringGaps = false;
+	FTimespan SendEmptySubtitleAt { FTimespan::MinValue() };
 };
diff --git a/Engine/Plugins/Media/ElectraSubtitles/Source/ElectraSubtitles/Private/wvtt/WebVTTParser.cpp b/Engine/Plugins/Media/ElectraSubtitles/Source/ElectraSubtitles/Private/wvtt/WebVTTParser.cpp
new file mode 100644
--- /dev/null
+++ b/Engine/Plugins/Media/ElectraSubtitles/Source/ElectraSubtitles/Private/wvtt/WebVTTParser.cpp
@@ -0,0 +1,1049 @@
+// Copyright Epic Games, Inc. All Rights Reserved.
+
+#include "WebVTTParser.h"
+#include "ElectraSubtitleUtils.h"
+#include "ParameterDictionary.h"
+
+namespace ElectraWebVTTParser
+{
+const TCHAR* const ConstWEBVTT(TEXT("WEBVTT"));
+const TCHAR* const ConstSTYLE(TEXT("STYLE"));
+const TCHAR* const ConstREGION(TEXT("REGION"));
+const TCHAR* const ConstNOTE(TEXT("NOTE"));
+const TCHAR* const ConstArrow(TEXT("-->"));
+
+const TCHAR* const ConstHLSTimestampMap(TEXT("X-TIMESTAMP-MAP"));
+const TCHAR* const ConstHLSTimestampLOCAL(TEXT("LOCAL"));
+const TCHAR* const ConstHLSTimestampMPEGTS(TEXT("MPEGTS"));
+
+const TCHAR* const ConstHLSHandleTimestampMapping(TEXT("handle_ts_mapping"));
+
+/*********************************************************************************************************************/
+
+FString IWebVTTParser::ProcessCueTextForType(FStringView InText, EWebVTTType InType)
+{
+	if (InText.Len() == 0)
+	{
+		return FString();
+	}
+	if (InType != EWebVTTType::Subtitle)
+	{
+		return FString(MoveTemp(InText));
+	}
+	// See if there is any '<' character that indicates a potential span
+	// or an '&' indicating a HTML character reference (like "&amp;" or "&gt;")
+	int32 LtPos = INDEX_NONE;
+	int32 AmpPos = INDEX_NONE;
+	if (!InText.FindChar(TCHAR('<'), LtPos) && !InText.FindChar(TCHAR('&'), AmpPos))
+	{
+		// With neither present we return an empty string to indicate the clean text
+		// is identical to the cue text. No need to duplicate anyting.
+		return FString(MoveTemp(InText));
+	}
+	TArray<TCHAR> New;
+	New.Reserve(InText.Len());
+	const TCHAR* Text = &InText[0];
+	const TCHAR* TextEnd = Text + InText.Len();
+	while(Text < TextEnd)
+	{
+		if (*Text == TCHAR('<'))
+		{
+			// Skip over whatever span this is.
+			while(Text < TextEnd && *Text != TCHAR('>'))
+			{
+				++Text;
+			}
+			++Text;
+		}
+		else if (*Text == TCHAR('&'))
+		{
+			const TCHAR* HtmlCharRef = Text + 1;
+			// Skip until we find the terminating ';'
+			while(Text < TextEnd && *Text != TCHAR(';'))
+			{
+				++Text;
+			}
+			if (Text < TextEnd)
+			{
+				FStringView CharRef = MakeStringView(HtmlCharRef, Text - HtmlCharRef);
+				++Text;
+
+				// We handle only very few named character references that are most commonly used.
+				// The list of all of them (https://html.spec.whatwg.org/multipage/named-characters.html#named-character-references)
+				// is too exhaustive to handle.
+				// We also completely ignore numeric character references.
+				if (CharRef.Equals(TEXT("lt"), ESearchCase::IgnoreCase))
+				{
+					New.Add(TCHAR('<'));
+				}
+				else if (CharRef.Equals(TEXT("gt"), ESearchCase::IgnoreCase))
+				{
+					New.Add(TCHAR('>'));
+				}
+				else if (CharRef.Equals(TEXT("amp"), ESearchCase::IgnoreCase))
+				{
+					New.Add(TCHAR('&'));
+				}
+				else if (CharRef.Equals(TEXT("quot"), ESearchCase::IgnoreCase))
+				{
+					New.Add(TCHAR('"'));
+				}
+				else if (CharRef.Equals(TEXT("nbsp"), ESearchCase::IgnoreCase))
+				{
+					New.Add(TCHAR(' '));
+				}
+			}
+		}
+		else
+		{
+			New.Add(*Text++);
+		}
+	}
+	return FString(MakeStringView(New.GetData(), New.Num()));
+}
+
+/*********************************************************************************************************************/
+
+class FWebVTTParser : public IWebVTTParser
+{
+public:
+	FWebVTTParser();
+	virtual ~FWebVTTParser();
+
+	const FString& GetLastErrorMessage() const override
+	{ return LastErrorMsg; }
+	bool ParseWebVTT(const TArray<uint8>& InWebVTTData, const Electra::FParamDict& InOptions) override;
+	void GetCuesAtTime(TArray<TUniquePtr<ICue>>& OutCues, FTimespan& OutNextChangeAt, TUniquePtr<ICueIterator>& InOutIterator, const FTimespan& InAtTime) override;
+
+	struct FWebVTTDocument
+	{
+		struct FMetadataHeader
+		{
+			FString Header;
+			FString Value;
+		};
+
+		struct FRegion
+		{
+			struct FAnchor
+			{
+				double x = 0.0;
+				double y = 100.0;
+			};
+			FString id;
+			double width = 100.0;
+			int32 lines = 3;
+			FAnchor regionAnchor;
+			FAnchor viewportAnchor;
+			bool bScrollUp = false;
+		};
+
+		struct FCue
+		{
+			FString ID;
+			FString Text;
+			FTimespan Start;
+			FTimespan End;
+			struct FSettings
+			{
+				enum class ELayout
+				{
+					Horizontal,
+					VerticalGrowingLeft,
+					VerticalGrowingRight
+				};
+				struct FLine
+				{
+					enum class EAlignment
+					{
+						Start,
+						Center,
+						End
+					};
+					double offset = 0.0;
+					EAlignment alignment = EAlignment::Start;
+					bool bIsPercentage = false;
+				};
+				struct FPosition
+				{
+					enum class EAlignment
+					{
+						Auto,
+						LineLeft,
+						Center,
+						LineRight
+					};
+					double position = 0.0;
+					EAlignment alignment = EAlignment::Auto;
+				};
+				enum class EAlignment
+				{
+					Start,
+					Center,
+					End,
+					Left,
+					Right
+				};
+				FLine line;
+				FPosition position;
+				FString region;
+				ELayout layout = ELayout::Horizontal;
+				EAlignment alignment = EAlignment::Center;
+				double size = 100.0;
+			};
+			FSettings Settings;
+		};
+
+		TArray<FMetadataHeader> MetadataHeaders;
+		TArray<FString> CSSStyles;
+		TArray<FRegion> Regions;
+		TArray<TSharedPtr<FCue, ESPMode::ThreadSafe>> Cues;
+		FTimespan SmallestCueStartTime { FTimespan::MaxValue() };
+		FTimespan LargestCueEndTime { FTimespan::MinValue() };
+
+		FTimespan TimestampOffset;
+
+	};
+
+
+	class FPublicCue : public ICue
+	{
+	public:
+		FPublicCue(const TSharedPtr<FWebVTTDocument::FCue, ESPMode::ThreadSafe>& InFromCue, const FTimespan& InTimeOffset)
+			: Cue(InFromCue), TimeOffset(InTimeOffset)
+		{ }
+		~FPublicCue() = default;
+		FTimespan GetStartTime() const override
+		{
+			return Cue->Start + TimeOffset;
+		}
+		FTimespan GetEndTime() const override
+		{
+			return Cue->End + TimeOffset;
+		}
+		FString GetID() const override
+		{
+			return Cue->ID;
+		}
+		FString GetText() const override
+		{
+			return Cue->Text;
+		}
+		bool operator == (const FPublicCue& Other) const
+		{
+			return Cue->Start == Other.Cue->Start && Cue->End == Other.Cue->End && Cue->ID.Equals(Other.Cue->ID) && Cue->Text.Equals(Other.Cue->Text);
+		}
+	private:
+		TSharedPtr<FWebVTTDocument::FCue, ESPMode::ThreadSafe> Cue;
+		FTimespan TimeOffset;
+	};
+
+	class FPublicCueIterator : public ICueIterator
+	{
+	public:
+		virtual ~FPublicCueIterator()
+		{ }
+
+		void Reset()
+		{
+			CurrentTime = FTimespan::MinValue();
+			CurrentIndex = 0;
+		}
+
+		FTimespan CurrentTime { FTimespan::MinValue() };
+		int32 CurrentIndex { 0 };
+	};
+
+
+private:
+
+	static bool IsAllNumeric(FStringView In)
+	{
+		for(auto& Ch : In)
+		{
+			if (!FChar::IsDigit(Ch))
+			{
+				return false;
+			}
+		}
+		return true;
+	}
+
+	static bool IsAllNumericMaybeNeg(FStringView In)
+	{
+		if (In[0] == TCHAR('-'))
+		{
+			for(int32 i=1; i<In.Len(); ++i)
+			{
+				if (!FChar::IsDigit(In[i]))
+				{
+					return false;
+				}
+			}
+			return true;
+		}
+		return IsAllNumeric(In);
+	}
+
+	static bool IsAllNumericDot(FStringView In)
+	{
+		int32 NumDots = 0;
+		for(auto& Ch : In)
+		{
+			NumDots += Ch == TCHAR('.') ? 1 : 0;
+			if (!FChar::IsDigit(Ch) && Ch != TCHAR('.'))
+			{
+				return false;
+			}
+		}
+		return NumDots <= 1;
+	}
+
+	static bool IsAllNumericMaybeNegDot(FStringView In)
+	{
+		int32 NumDots = 0;
+		for(int32 i=0; i<In.Len(); ++i)
+		{
+			if (i == 0 && In[0] == TCHAR('-'))
+			{
+				continue;
+			}
+			NumDots += In[i] == TCHAR('.') ? 1 : 0;
+			if (!FChar::IsDigit(In[i]) && In[i] != TCHAR('.'))
+			{
+				return false;
+			}
+		}
+		return NumDots <= 1;
+	}
+
+	static bool IsNewline(const FString::TConstIterator& InIt)
+	{
+		return *InIt == TCHAR('\n') || *InIt == TCHAR('\r');
+	}
+
+	static bool LocateEndOfLine(FString::TConstIterator& InOutIt)
+	{
+		bool bHasNonWS = false;
+		while(InOutIt && !IsNewline(InOutIt))
+		{
+			bHasNonWS = bHasNonWS || !FChar::IsWhitespace(*InOutIt);
+			++InOutIt;
+		}
+		return !bHasNonWS;
+	}
+
+	static int32 SkipOverLinebreaks(FString::TConstIterator& InOutIt)
+	{
+		int32 NumEOLs = 0;
+		while(InOutIt)
+		{
+			if (*InOutIt == TCHAR('\r'))
+			{
+				++NumEOLs;
+				++InOutIt;
+				if (InOutIt && *InOutIt == TCHAR('\n'))
+				{
+					++InOutIt;
+				}
+			}
+			else if (*InOutIt == TCHAR('\n'))
+			{
+				++NumEOLs;
+				++InOutIt;
+			}
+			else
+			{
+				break;
+			}
+		}
+		return NumEOLs;
+	}
+
+	static FStringView CreateView(const FString::TConstIterator& sIt, const FString::TConstIterator& eIt)
+	{
+		return MakeStringView(sIt.operator->(), eIt.GetIndex() - sIt.GetIndex());
+	}
+
+	static int32 SplitAtAndRemoveWhitespace(TArray<FStringView>& OutParts, bool& bInOutContainsArrow, FStringView In)
+	{
+		OutParts.Empty();
+		bInOutContainsArrow = false;
+		int32 sIdx = -1;
+		for(int32 i=0; i<In.Len(); ++i)
+		{
+			if (!FChar::IsWhitespace(In[i]))
+			{
+				if (sIdx < 0)
+				{
+					sIdx = i;
+				}
+			}
+			else if (sIdx >= 0)
+			{
+				int32 Len = i - sIdx;
+				OutParts.Emplace(MakeStringView(&In[sIdx], Len));
+				bInOutContainsArrow = bInOutContainsArrow || OutParts.Last().Equals(ConstArrow);
+				sIdx = -1;
+			}
+		}
+		if (sIdx >= 0)
+		{
+			OutParts.Emplace(MakeStringView(&In[sIdx], In.Len() - sIdx));
+			bInOutContainsArrow = bInOutContainsArrow || OutParts.Last().Equals(ConstArrow);
+		}
+		return OutParts.Num();
+	}
+
+	static bool ParseLineIntoPartsAtWhitespace(TArray<FStringView>& OutParts, FString::TConstIterator& InOutIt)
+	{
+		FString::TConstIterator Start(InOutIt);
+		LocateEndOfLine(InOutIt);
+		bool bContainsArrow;
+		SplitAtAndRemoveWhitespace(OutParts, bContainsArrow, CreateView(Start, InOutIt));
+		return bContainsArrow;
+	}
+
+	static bool ParseTimestamp(FTimespan& OutTime, FStringView In)
+	{
+		FString Time(In);
+		TArray<FString> Parts;
+		Time.ParseIntoArray(Parts, TEXT(":"), false);
+		if (Parts.Num() < 2 || Parts.Num() > 3)
+		{
+			return false;
+		}
+		// Minutes, Seconds and fractional seconds are mandatory.
+		// If there are no hours we add a 0 for hours to simplify parsing.
+		if (Parts.Num() == 2)
+		{
+			Parts.Insert(TEXT("0"), 0);
+		}
+		if (!IsAllNumeric(Parts[0]) || !IsAllNumeric(Parts[1]) || !IsAllNumericDot(Parts[2]))
+		{
+			return false;
+		}
+		int32 DotPos;
+		if (!Parts[2].FindChar(TCHAR('.'), DotPos) || DotPos != 2)
+		{
+			return false;
+		}
+		Parts.Add(Parts[2].Mid(DotPos + 1));
+		Parts[2].LeftInline(DotPos);
+		if (Parts[1].Len() != 2 || Parts[2].Len() != 2 || Parts[3].Len() != 3)
+		{
+			return false;
+		}
+		int32 h = FCString::Atoi(*Parts[0]);
+		int32 m = FCString::Atoi(*Parts[1]);
+		int32 s = FCString::Atoi(*Parts[2]);
+		int32 ms = FCString::Atoi(*Parts[3]);
+		if (m > 59 || s > 59)
+		{
+			return false;
+		}
+		OutTime = FTimespan(0, h, m, s, ms * 1000000);
+		return true;
+	}
+
+
+	TUniquePtr<FWebVTTDocument> CurrentDoc;
+
+	FString LastErrorMsg;
+	bool bInCues = false;
+};
+
+
+/*********************************************************************************************************************/
+
+TSharedPtr<IWebVTTParser, ESPMode::ThreadSafe> IWebVTTParser::Create()
+{
+	return MakeShared<FWebVTTParser, ESPMode::ThreadSafe>();
+}
+
+FWebVTTParser::FWebVTTParser()
+{
+}
+
+FWebVTTParser::~FWebVTTParser()
+{
+}
+
+
+
+
+bool FWebVTTParser::ParseWebVTT(const TArray<uint8>& InWebVTTData, const Electra::FParamDict& InOptions)
+{
+	if (InWebVTTData.Num() < 6)
+	{
+		return false;
+	}
+	bool bHaveBOM = InWebVTTData[0] == 0xef && InWebVTTData[1] == 0xbb && InWebVTTData[2] == 0xbf;
+	// The data is expected to be a UTF-8 encoded string. If there is a BOM we skip over it.
+	FString Document = FString::ConstructFromPtrSize((const UTF8CHAR*)InWebVTTData.GetData() + (bHaveBOM ? 3 : 0), InWebVTTData.Num() - (bHaveBOM ? 3 : 0));
+	if (Document.Len() < 6)
+	{
+		return false;
+	}
+
+	// Set up iterator to the start of the document and start checking it to contain a proper signature
+	TArray<FStringView> Parts;
+	FString::TConstIterator It(Document.CreateConstIterator());
+	ParseLineIntoPartsAtWhitespace(Parts, It);
+	if (Parts.IsEmpty() || Parts.Num() > 1 || !Parts[0].Equals(ConstWEBVTT, ESearchCase::CaseSensitive))
+	{
+		return false;
+	}
+
+	TUniquePtr<FWebVTTDocument> NewDoc = MakeUnique<FWebVTTDocument>();
+
+	// We do not check the document for validity, just parse it line for line.
+	FString PreviousUnknownLine;
+	bool bSkipUntilEmptyLine = false;
+	while(It)
+	{
+		if (bSkipUntilEmptyLine)
+		{
+			// The assumption is that the iterator is currently pointing to the linebreak
+			// at the end of the line asking to skip ahead.
+			// So we first skip over linebreaks and if that's more than just one we are done.
+			// If not we check if the line is empty (all whitespaces) which is also good.
+			if (SkipOverLinebreaks(It) > 1 || LocateEndOfLine(It))
+			{
+				bSkipUntilEmptyLine = false;
+			}
+			continue;
+		}
+
+		// Move on to the next line by skipping over CR/LF sequences. This does not skip
+		// over lines that contain whitespace.
+		int32 NumLineBreaks = SkipOverLinebreaks(It);
+		if (It)
+		{
+			// If we skipped more than one line break then the previous line is of no interest any more.
+			if (NumLineBreaks > 1)
+			{
+				PreviousUnknownLine.Empty();
+			}
+			FString::TConstIterator LineStart(It);
+			bool bContainsArrow = ParseLineIntoPartsAtWhitespace(Parts, It);
+			if (bContainsArrow)
+			{
+				if (Parts.Num() < 3 || !Parts[1].Equals(ConstArrow))
+				{
+					FString cue(CreateView(LineStart, It));
+					LastErrorMsg = FString::Printf(TEXT("\"%s\" is not a cue"), *cue);
+					return false;
+				}
+				else
+				{
+					bInCues = true;
+
+					TSharedPtr<FWebVTTDocument::FCue, ESPMode::ThreadSafe> NewCue = MakeShared<FWebVTTDocument::FCue, ESPMode::ThreadSafe>();
+					NewCue->ID = MoveTemp(PreviousUnknownLine);
+					if (!ParseTimestamp(NewCue->Start, Parts[0]))
+					{
+						LastErrorMsg = FString::Printf(TEXT("\"%s\" is not a valid timestamp"), *FString(Parts[0]));
+						return false;
+					}
+					if (!ParseTimestamp(NewCue->End, Parts[2]))
+					{
+						LastErrorMsg = FString::Printf(TEXT("\"%s\" is not a valid timestamp"), *FString(Parts[1]));
+						return false;
+					}
+					// Parse cue settings
+					for(int32 np=3; np<Parts.Num(); ++np)
+					{
+						int32 ColonPos = INDEX_NONE;
+						if (Parts[np].FindChar(TCHAR(':'), ColonPos) && ColonPos > 0 && ColonPos+1 < Parts[np].Len())
+						{
+							auto Key = Parts[np].Left(ColonPos);
+							auto Value = Parts[np].Mid(ColonPos + 1);
+							if (Key.Equals(TEXT("region"), ESearchCase::CaseSensitive))
+							{
+								NewCue->Settings.region = Value;
+							}
+							else if (Key.Equals(TEXT("vertical"), ESearchCase::CaseSensitive))
+							{
+								if (Value.Equals(TEXT("lr"), ESearchCase::CaseSensitive))
+								{
+									NewCue->Settings.layout = FWebVTTDocument::FCue::FSettings::ELayout::VerticalGrowingRight;
+								}
+								else if (Value.Equals(TEXT("rl"), ESearchCase::CaseSensitive))
+								{
+									NewCue->Settings.layout = FWebVTTDocument::FCue::FSettings::ELayout::VerticalGrowingLeft;
+								}
+							}
+							else if (Key.Equals(TEXT("line"), ESearchCase::CaseSensitive))
+							{
+								FString SubPart(Value);
+								TArray<FString> SubParts;
+								SubPart.ParseIntoArray(SubParts, TEXT(","), true);
+								bool bCont = false;
+								if (SubParts[0].EndsWith(TEXT("%")))
+								{
+									SubParts[0].LeftChopInline(1);
+									if (IsAllNumericDot(SubParts[0]))
+									{
+										double Pct = FCString::Atod(*SubParts[0]);
+										if (Pct >= 0.0 && Pct <= 100.0)
+										{
+											NewCue->Settings.line.offset = Pct;
+											NewCue->Settings.line.bIsPercentage = true;
+											bCont = true;
+										}
+									}
+								}
+								else if (IsAllNumericMaybeNegDot(SubParts[0]))
+								{
+									NewCue->Settings.line.offset = FCString::Atod(*SubParts[0]);
+									NewCue->Settings.line.bIsPercentage = false;
+									bCont = true;
+								}
+								if (bCont && SubParts.Num() > 1)
+								{
+									if (SubParts[1].Equals(TEXT("start"), ESearchCase::CaseSensitive))
+									{
+										NewCue->Settings.line.alignment = FWebVTTDocument::FCue::FSettings::FLine::EAlignment::Start;
+									}
+									else if (SubParts[1].Equals(TEXT("center"), ESearchCase::CaseSensitive))
+									{
+										NewCue->Settings.line.alignment = FWebVTTDocument::FCue::FSettings::FLine::EAlignment::Center;
+									}
+									else if (SubParts[1].Equals(TEXT("end"), ESearchCase::CaseSensitive))
+									{
+										NewCue->Settings.line.alignment = FWebVTTDocument::FCue::FSettings::FLine::EAlignment::End;
+									}
+								}
+							}
+							else if (Key.Equals(TEXT("position"), ESearchCase::CaseSensitive))
+							{
+								FString SubPart(Value);
+								TArray<FString> SubParts;
+								SubPart.ParseIntoArray(SubParts, TEXT(","), true);
+								if (SubParts[0].EndsWith(TEXT("%")))
+								{
+									SubParts[0].LeftChopInline(1);
+									if (IsAllNumericDot(SubParts[0]))
+									{
+										double Pct = FCString::Atod(*SubParts[0]);
+										if (Pct >= 0.0 && Pct <= 100.0)
+										{
+											NewCue->Settings.position.position = Pct;
+											if (SubParts.Num() > 1)
+											{
+												if (SubParts[1].Equals(TEXT("line-left"), ESearchCase::CaseSensitive))
+												{
+													NewCue->Settings.position.alignment = FWebVTTDocument::FCue::FSettings::FPosition::EAlignment::LineLeft;
+												}
+												else if (SubParts[1].Equals(TEXT("center"), ESearchCase::CaseSensitive))
+												{
+													NewCue->Settings.position.alignment = FWebVTTDocument::FCue::FSettings::FPosition::EAlignment::Center;
+												}
+												else if (SubParts[1].Equals(TEXT("line-right"), ESearchCase::CaseSensitive))
+												{
+													NewCue->Settings.position.alignment = FWebVTTDocument::FCue::FSettings::FPosition::EAlignment::LineRight;
+												}
+											}
+										}
+									}
+								}
+							}
+							else if (Key.Equals(TEXT("align"), ESearchCase::CaseSensitive))
+							{
+								if (Value.Equals(TEXT("start"), ESearchCase::CaseSensitive))
+								{
+									NewCue->Settings.alignment = FWebVTTDocument::FCue::FSettings::EAlignment::Start;
+								}
+								else if (Value.Equals(TEXT("center"), ESearchCase::CaseSensitive))
+								{
+									NewCue->Settings.alignment = FWebVTTDocument::FCue::FSettings::EAlignment::Center;
+								}
+								else if (Value.Equals(TEXT("end"), ESearchCase::CaseSensitive))
+								{
+									NewCue->Settings.alignment = FWebVTTDocument::FCue::FSettings::EAlignment::Right;
+								}
+								else if (Value.Equals(TEXT("left"), ESearchCase::CaseSensitive))
+								{
+									NewCue->Settings.alignment = FWebVTTDocument::FCue::FSettings::EAlignment::Left;
+								}
+								else if (Value.Equals(TEXT("right"), ESearchCase::CaseSensitive))
+								{
+									NewCue->Settings.alignment = FWebVTTDocument::FCue::FSettings::EAlignment::Right;
+								}
+							}
+							else if (Key.Equals(TEXT("size"), ESearchCase::CaseSensitive))
+							{
+								if (Value.EndsWith(TEXT("%")))
+								{
+									FString Num(Value);
+									Num.LeftChopInline(1);
+									if (IsAllNumericDot(Num))
+									{
+										double Pct = FCString::Atod(*Num);
+										if (Pct >= 0.0 && Pct <= 100.0)
+										{
+											NewCue->Settings.size = Pct;
+										}
+									}
+								}
+							}
+						}
+					}
+
+					// Parse the actual cue text now
+					while(1)
+					{
+						// Skip over the line break at the end of the line.
+						if (SkipOverLinebreaks(It) > 1)
+						{
+							// If there was more than one line break the cue is empty.
+							break;
+						}
+						FString::TConstIterator CueLineStart(It);
+						// Go to the end of the line. If it is empty we are done.
+						if (LocateEndOfLine(It))
+						{
+							break;
+						}
+						// If there is text in the cue we need to append a line break
+						if (NewCue->Text.Len())
+						{
+							NewCue->Text.Append(TEXT("\n"));
+						}
+						FStringView Text(CreateView(CueLineStart, It));
+						NewCue->Text.Append(Text);
+					}
+
+					// Add the new cue to the list.
+					if (NewCue->Start < NewDoc->SmallestCueStartTime)
+					{
+						NewDoc->SmallestCueStartTime = NewCue->Start;
+					}
+					if (NewCue->End > NewDoc->LargestCueEndTime)
+					{
+						NewDoc->LargestCueEndTime = NewCue->End;
+					}
+					NewDoc->Cues.Emplace(MoveTemp(NewCue));
+				}
+			}
+			else if (Parts.Num())
+			{
+				PreviousUnknownLine.Empty();
+				// Is this a comment?
+				if (Parts[0].Equals(ConstNOTE, ESearchCase::CaseSensitive))
+				{
+					// Comments just get skipped.
+					bSkipUntilEmptyLine = true;
+				}
+				// A style block?
+				else if (Parts[0].Equals(ConstSTYLE, ESearchCase::CaseSensitive))
+				{
+					// Skip over the line break at the end of the line.
+					if (SkipOverLinebreaks(It) > 1)
+					{
+						// If there was more than one line break the style block is empty.
+					}
+					else
+					{
+						// We need to collect anything until an empty line, but not parse it.
+						// This is pure CSS which is not meant for us to interpret.
+						FString::TConstIterator CSSStart(It);
+						while(1)
+						{
+							bool bEmptyLine = LocateEndOfLine(It);
+							if (bEmptyLine || SkipOverLinebreaks(It) > 1)
+							{
+								break;
+							}
+						}
+						if (It)
+						{
+							NewDoc->CSSStyles.Emplace(CreateView(CSSStart, It));
+						}
+					}
+				}
+				// A region block?
+				else if (Parts[0].Equals(ConstREGION, ESearchCase::CaseSensitive))
+				{
+					// Skip over the line break at the end of the line.
+					if (SkipOverLinebreaks(It) > 1)
+					{
+						// If there was more than one line break the region block is empty.
+					}
+					else
+					{
+						FWebVTTDocument::FRegion NewRegion;
+						while(1)
+						{
+							ParseLineIntoPartsAtWhitespace(Parts, It);
+							if (Parts.IsEmpty())
+							{
+								break;
+							}
+							for(int32 i=0; i<Parts.Num(); ++i)
+							{
+								int32 ColonPos = INDEX_NONE;
+								if (Parts[i].FindChar(TCHAR(':'), ColonPos) && ColonPos > 0 && ColonPos+1 < Parts[i].Len())
+								{
+									auto Key = Parts[i].Left(ColonPos);
+									auto Value = Parts[i].Mid(ColonPos + 1);
+									if (Key.Equals(TEXT("id"), ESearchCase::CaseSensitive))
+									{
+										NewRegion.id = Value;
+									}
+									else if (Key.Equals(TEXT("width"), ESearchCase::CaseSensitive))
+									{
+										if (Value[Value.Len() - 1] == TCHAR('%'))
+										{
+											auto Num = Value.Left(Value.Len() - 1);
+											if (IsAllNumericDot(Num))
+											{
+												double Pct = FCString::Atod(&Num[0]);
+												if (Pct >= 0.0 && Pct <= 100.0)
+												{
+													NewRegion.width = Pct;
+												}
+											}
+										}
+									}
+									else if (Key.Equals(TEXT("lines"), ESearchCase::CaseSensitive))
+									{
+										if (IsAllNumeric(Value))
+										{
+											NewRegion.lines = FCString::Atoi(&Value[0]);
+										}
+									}
+									else if (Key.Equals(TEXT("regionanchor"), ESearchCase::CaseSensitive) ||
+											 Key.Equals(TEXT("viewportanchor"), ESearchCase::CaseSensitive))
+									{
+										FString SubPart(Value);
+										TArray<FString> SubParts;
+										SubPart.ParseIntoArray(SubParts, TEXT(","), true);
+										if (SubParts.Num() == 2 && SubParts[0].EndsWith(TEXT("%")) && SubParts[1].EndsWith(TEXT("%")))
+										{
+											SubParts[0].LeftChopInline(1);
+											SubParts[1].LeftChopInline(1);
+											if (IsAllNumericDot(SubParts[0]) && IsAllNumericDot(SubParts[1]))
+											{
+												double PctX = FCString::Atod(*SubParts[0]);
+												double PctY = FCString::Atod(*SubParts[1]);
+												if (PctX >= 0.0 && PctX <= 100.0 && PctY >= 0.0 && PctY <= 100.0)
+												{
+													if (Key[0] == TCHAR('r'))
+													{
+														NewRegion.regionAnchor.x = PctX;
+														NewRegion.regionAnchor.y = PctY;
+													}
+													else
+													{
+														NewRegion.viewportAnchor.x = PctX;
+														NewRegion.viewportAnchor.y = PctY;
+													}
+												}
+											}
+										}
+									}
+									else if (Key.Equals(TEXT("scroll"), ESearchCase::CaseSensitive))
+									{
+										NewRegion.bScrollUp = Value.Equals(TEXT("up"), ESearchCase::CaseSensitive);
+									}
+								}
+							}
+							if (SkipOverLinebreaks(It) > 1)
+							{
+								break;
+							}
+						}
+						NewDoc->Regions.Emplace(MoveTemp(NewRegion));
+					}
+				}
+				else
+				{
+					// Something else.
+					// Check if this is a (deprecated since 2016 but still in use with old files) metadata header.
+					// Metadata headers cannot appear once a cue has been found.
+					auto Line(CreateView(LineStart, It));
+
+					/*
+						Note: The old WebVTT spec that still contained metadata headers says that the header
+						      and value are separated by a COLON
+						        https://www.w3.org/TR/2015/WD-webvtt1-20151208/#webvtt-metadata-header
+						      which I'm sure is right. The only metadata we are interested in at this point
+							  is Apple's HLS "X-TIMESTAMP-MAP" and this uses an EQUALS sign.
+							  Yes, there are colons in there as well, but apparently not to separate header
+							  and value.
+							  We have seen differently structured headers as well like:
+									X-TIMESTAMP-MAP=LOCAL:00:00:00.000,MPEGTS:9000
+								and
+									X-TIMESTAMP-MAP=MPEGTS:900000,LOCAL:00:00:00.000
+
+							  so clearly the separator here is an equals sign.
+					*/
+
+					int32 ColonPos, EqualsPos;
+					Line.FindChar(TCHAR('='), EqualsPos);
+					Line.FindChar(TCHAR(':'), ColonPos);
+					// We look at both colon and equal and use whichever one appears first.
+					if (!bInCues && (ColonPos != INDEX_NONE || EqualsPos != INDEX_NONE))
+					{
+						int32 HeaderSepPos = (ColonPos != INDEX_NONE && EqualsPos != INDEX_NONE) ? (EqualsPos < ColonPos ? EqualsPos : ColonPos) : (ColonPos != INDEX_NONE ? ColonPos : EqualsPos);
+						// The old WebVTT spec says nothing about skipping leading or trailing whitespace
+						// so we leave them in!
+						FWebVTTDocument::FMetadataHeader& Hdr(NewDoc->MetadataHeaders.Emplace_GetRef());
+						Hdr.Header = Line.Left(HeaderSepPos);
+						Hdr.Value = Line.Mid(HeaderSepPos + 1);
+					}
+					else
+					{
+						PreviousUnknownLine = Line;
+						PreviousUnknownLine.TrimStartAndEndInline();
+					}
+				}
+			}
+		}
+	}
+
+	// Check if there is a timestamp mapping as it may occur with HLS.
+	if (InOptions.GetValue(ElectraWebVTTParser::ConstHLSHandleTimestampMapping).SafeGetBool())
+	{
+		/*
+			Look for the mapping as described in the HLS RFC.
+			There are two versions of this around:
+				X-TIMESTAMP-MAP=LOCAL:00:00:00.000,MPEGTS:9000
+			and
+				X-TIMESTAMP-MAP=MPEGTS:900000,LOCAL:00:00:00.000
+		*/
+		const FWebVTTDocument::FMetadataHeader* HLSTimestampMapping = NewDoc->MetadataHeaders.FindByPredicate([](const FWebVTTDocument::FMetadataHeader& h){ return h.Header.Equals(ConstHLSTimestampMap); });
+		if (HLSTimestampMapping)
+		{
+			TArray<FString> MappingParts;
+			HLSTimestampMapping->Value.ParseIntoArray(MappingParts, TEXT(","), true);
+			if (MappingParts.Num() == 2)
+			{
+				int32 ColonPos1, ColonPos2;
+				MappingParts[0].FindChar(TCHAR(':'), ColonPos1);
+				MappingParts[1].FindChar(TCHAR(':'), ColonPos2);
+				if (ColonPos1 != INDEX_NONE && ColonPos2 != INDEX_NONE)
+				{
+					TArray<FString> Values[2];
+					Values[0].Emplace(MappingParts[0].Left(ColonPos1));
+					Values[0].Emplace(MappingParts[0].Mid(ColonPos1 + 1));
+					Values[1].Emplace(MappingParts[1].Left(ColonPos2));
+					Values[1].Emplace(MappingParts[1].Mid(ColonPos2 + 1));
+					int32 LocalIdx = Values[0][0].Equals(ConstHLSTimestampLOCAL) ? 0 : Values[1][0].Equals(ConstHLSTimestampLOCAL) ? 1 : -1;
+					int32 MpegIdx = Values[0][0].Equals(ConstHLSTimestampMPEGTS) ? 0 : Values[1][0].Equals(ConstHLSTimestampMPEGTS) ? 1 : -1;
+					if (LocalIdx != -1 && MpegIdx != -1)
+					{
+						FTimespan LocalTime;
+						// Parse the local time and check the Mpeg time
+						if (ParseTimestamp(LocalTime, Values[LocalIdx][1]) && IsAllNumeric(Values[MpegIdx][1]))
+						{
+							// The value is given in the typical MPEG-TS 90kHz clock. Convert it into HNS as used by FTimespan.
+							FTimespan MpegTS90(FCString::Atoi64(*Values[MpegIdx][1]) * 1000 / 9);
+
+							// The mapping maps a local cue time to some MPEG-TS time.
+							// Typically the local time is zero, but could be anything.
+							// The offset to be added to all the local times is the difference between the MPEG-TS time and the local cue time.
+							NewDoc->TimestampOffset = MpegTS90 - LocalTime;
+						}
+					}
+				}
+			}
+		}
+	}
+
+	// Set the new document as the current one.
+	CurrentDoc = MoveTemp(NewDoc);
+	return true;
+}
+
+
+
+void FWebVTTParser::GetCuesAtTime(TArray<TUniquePtr<ICue>>& OutCues, FTimespan& OutNextChangeAt, TUniquePtr<ICueIterator>& InOutIterator, const FTimespan& InAtTime)
+{
+	if (!CurrentDoc.IsValid() || InAtTime > CurrentDoc->LargestCueEndTime + CurrentDoc->TimestampOffset)
+	{
+		InOutIterator.Reset();
+		OutNextChangeAt = FTimespan::MaxValue();
+		return;
+	}
+	else if (InAtTime < CurrentDoc->SmallestCueStartTime + CurrentDoc->TimestampOffset)
+	{
+		InOutIterator.Reset();
+		OutNextChangeAt = CurrentDoc->SmallestCueStartTime + CurrentDoc->TimestampOffset;
+		return;
+	}
+
+	// If the iterator is not set up yet, locate the first cue.
+	if (!InOutIterator.IsValid())
+	{
+		InOutIterator = MakeUnique<FPublicCueIterator>();
+	}
+	FPublicCueIterator* It = static_cast<FPublicCueIterator*>(InOutIterator.Get());
+	if (InAtTime < It->CurrentTime)
+	{
+		It->Reset();
+	}
+	if (It->CurrentTime == FTimespan::MinValue())
+	{
+		for(int32 i=0,iMax=CurrentDoc->Cues.Num(); i<iMax; ++i)
+		{
+			TRange<FTimespan> CueRange(CurrentDoc->Cues[i]->Start + CurrentDoc->TimestampOffset, CurrentDoc->Cues[i]->End + CurrentDoc->TimestampOffset);
+			if (CueRange.Contains(InAtTime))
+			{
+				It->CurrentTime = CueRange.GetLowerBoundValue() + CurrentDoc->TimestampOffset;
+				It->CurrentIndex = i;
+				break;
+			}
+		}
+		check(It->CurrentIndex >= 0);
+		if (It->CurrentIndex < 0)
+		{
+			return;
+		}
+	}
+	// Look at the next cues starting at the iterator's position.
+	FTimespan NextChangeOccursAt = FTimespan::MaxValue();
+	for(int32 i=It->CurrentIndex,iMax=CurrentDoc->Cues.Num(); i<iMax; ++i)
+	{
+		TRange<FTimespan> CueRange(CurrentDoc->Cues[i]->Start + CurrentDoc->TimestampOffset, CurrentDoc->Cues[i]->End + CurrentDoc->TimestampOffset);
+		if (CueRange.Contains(InAtTime))
+		{
+			TUniquePtr<ICue> Cue(new FPublicCue(CurrentDoc->Cues[i], CurrentDoc->TimestampOffset));
+			bool bIsDuplicate = false;
+			for(int nd=0, ndMax=OutCues.Num(); nd<ndMax; ++nd)
+			{
+				if (*static_cast<const FPublicCue*>(OutCues[nd].Get()) == *static_cast<const FPublicCue*>(Cue.Get()))
+				{
+					bIsDuplicate = true;
+					break;
+				}
+			}
+			if (!bIsDuplicate)
+			{
+				OutCues.AddUnique(MoveTemp(Cue));
+			}
+			NextChangeOccursAt = FMath::Min(CueRange.GetUpperBoundValue(), NextChangeOccursAt);
+		}
+		else if (CueRange.GetLowerBoundValue() > InAtTime)
+		{
+			if (OutCues.IsEmpty())
+			{
+				NextChangeOccursAt = CueRange.GetLowerBoundValue();
+			}
+			break;
+		}
+	}
+	OutNextChangeAt = NextChangeOccursAt;
+
+	// Skip the iterator ahead over all cues that are now in the past.
+	while(It->CurrentIndex < CurrentDoc->Cues.Num() && CurrentDoc->Cues[It->CurrentIndex]->End + CurrentDoc->TimestampOffset < InAtTime)
+	{
+		++It->CurrentIndex;
+		It->CurrentTime = (It->CurrentIndex < CurrentDoc->Cues.Num() ? CurrentDoc->Cues[It->CurrentIndex]->Start : CurrentDoc->LargestCueEndTime) + CurrentDoc->TimestampOffset;
+	}
+}
+
+
+} // namespace ElectraWebVTTParser
diff --git a/Engine/Plugins/Media/ElectraSubtitles/Source/ElectraSubtitles/Private/wvtt/WebVTTParser.h b/Engine/Plugins/Media/ElectraSubtitles/Source/ElectraSubtitles/Private/wvtt/WebVTTParser.h
new file mode 100644
--- /dev/null
+++ b/Engine/Plugins/Media/ElectraSubtitles/Source/ElectraSubtitles/Private/wvtt/WebVTTParser.h
@@ -0,0 +1,61 @@
+// Copyright Epic Games, Inc. All Rights Reserved.
+
+#pragma once
+
+#include "Templates/SharedPointer.h"
+
+namespace Electra { class FParamDict; }
+namespace Electra { class FTimeValue; }
+
+namespace ElectraWebVTTParser
+{
+
+class IWebVTTParser
+{
+public:
+	static TSharedPtr<IWebVTTParser, ESPMode::ThreadSafe> Create();
+
+	virtual ~IWebVTTParser() = default;
+
+	/**
+	 * Returns the most recent error message.
+	 */
+	virtual const FString& GetLastErrorMessage() const = 0;
+
+	/**
+	 * Parses the provided WebVTT document.
+	 *
+	 * @return true if successful, false if not. Get error message from GetLastErrorMessage().
+	 */
+	virtual bool ParseWebVTT(const TArray<uint8>& InWebVTTData, const Electra::FParamDict& InOptions) = 0;
+
+
+	class ICue
+	{
+	public:
+		virtual ~ICue() = default;
+		virtual FTimespan GetStartTime() const = 0;
+		virtual FTimespan GetEndTime() const = 0;
+		virtual FString GetID() const = 0;
+		virtual FString GetText() const = 0;
+	};
+
+	class ICueIterator
+	{
+	public:
+		virtual ~ICueIterator() = default;
+	};
+
+	virtual void GetCuesAtTime(TArray<TUniquePtr<ICue>>& OutCues, FTimespan& OutNextChangeAt, TUniquePtr<ICueIterator>& InOutIterator, const FTimespan& InAtTime) = 0;
+
+
+	enum class EWebVTTType
+	{
+		Subtitle,
+		Chapters,
+		Metadata
+	};
+	static FString ProcessCueTextForType(FStringView InText, EWebVTTType InType);
+};
+
+}
