ElectraSubtitles: Fix for overlapping WebVTT subtitles

#rnx

diff --git a/Engine/Plugins/Media/ElectraSubtitles/Source/ElectraSubtitles/Private/wvtt/ElectraSubtitleDecoder_WVTT.cpp b/Engine/Plugins/Media/ElectraSubtitles/Source/ElectraSubtitles/Private/wvtt/ElectraSubtitleDecoder_WVTT.cpp
--- a/Engine/Plugins/Media/ElectraSubtitles/Source/ElectraSubtitles/Private/wvtt/ElectraSubtitleDecoder_WVTT.cpp
+++ b/Engine/Plugins/Media/ElectraSubtitles/Source/ElectraSubtitles/Private/wvtt/ElectraSubtitleDecoder_WVTT.cpp
@@ -242,7 +242,7 @@ public:
 			&& (bIsGeneratedID || Other.bIsGeneratedID || ID.Equals(Other.ID))
 			&& Duration == Other.Duration && TextAsArray == Other.TextAsArray;
 	}
-	
+
 	bool operator != (const FSubtitleDecoderOutputWVTT& Other) const
 	{
 		return !(this->operator==(Other));
@@ -372,7 +372,7 @@ void FElectraSubtitleDecoderWVTT::AddStreamedSubtitleData(const TArray<uint8>& I
 			}
 		}
 		ParsedTimerangeLock.Unlock();
-		
+
 		TSharedPtr<FParsedTimerange, ESPMode::ThreadSafe> ParsedTimerange = MakeShared<FParsedTimerange, ESPMode::ThreadSafe>();
 		ParsedTimerange->Parser = ElectraWebVTTParser::IWebVTTParser::Create();
 		ParsedTimerange->AbsoluteStartTime = InAbsoluteTimestamp;
@@ -577,7 +577,7 @@ void FElectraSubtitleDecoderWVTT::UpdatePlaybackPosition(Electra::FTimeValue InA
 	TArray<TSharedPtr<FParsedTimerange, ESPMode::ThreadSafe>> TimeRanges;
 	for(int32 i=0; i<ParsedTimeranges.Num(); ++i)
 	{
-		if (InAbsolutePosition > ParsedTimeranges[i]->AbsoluteStartTime + ParsedTimeranges[i]->Duration && 
+		if (InAbsolutePosition > ParsedTimeranges[i]->AbsoluteStartTime + ParsedTimeranges[i]->Duration &&
 			InAbsolutePosition.GetSequenceIndex() >= ParsedTimeranges[i]->AbsoluteStartTime.GetSequenceIndex())
 		{
 			ParsedTimeranges.RemoveAt(i);
@@ -605,6 +605,7 @@ void FElectraSubtitleDecoderWVTT::UpdatePlaybackPosition(Electra::FTimeValue InA
 	}
 	LastPlaybackUpdateAbsPos = InAbsolutePosition;
 
+	FTimespan ThisEvalTime = NextEvaluationAt;
 	FTimespan NextEval = NextEvaluationAt;
 	bool bWasForcedEval = NextEvaluationAt == FTimespan::MinValue();
 	ParsedTimerangeLock.Unlock();
@@ -628,18 +629,26 @@ void FElectraSubtitleDecoderWVTT::UpdatePlaybackPosition(Electra::FTimeValue InA
 	ParsedTimerangeLock.Unlock();
 	if (Cues.Num())
 	{
+		// Start with this evaluation time as the new cue's start time.
+		// If we have overlapping subtitles like
+		// [a...........b]
+		//       [c..d]
+		// We need to split this into
+		// [a..c], [c..d], [d..b]  and NOT
+		// [a..c], [c..d], [a..b]  !!!
+		FTimespan LargestCueStartTime = ThisEvalTime;
+		FTimespan LargestCueEndTime = FTimespan::MinValue();
 		FString CombinedPlainText;
-		FTimespan CueStartTime = FTimespan::MaxValue();
-		FTimespan CueEndTime = FTimespan::MaxValue();
 		for(int32 i=0; i<Cues.Num(); ++i)
 		{
-			CueStartTime = FMath::Min(CueStartTime, Cues[i]->GetStartTime());
-			CueEndTime = FMath::Min(CueEndTime, Cues[i]->GetEndTime());
+			// Get the largest start time of the currently active cues.
+			// This will be the start time of the new subtitle we send out.
+			LargestCueStartTime = FMath::Max(LargestCueStartTime, Cues[i]->GetStartTime());
+			LargestCueEndTime = FMath::Max(LargestCueStartTime, Cues[i]->GetEndTime());
 			if (i > 0)
 			{
 				CombinedPlainText.Append(TEXT("\n"));
 			}
-
 			CombinedPlainText.Append(ElectraWebVTTParser::IWebVTTParser::ProcessCueTextForType(Cues[i]->GetText(), ElectraWebVTTParser::IWebVTTParser::EWebVTTType::Subtitle));
 		}
 		// Get the cue ID only when there is a single cue. If we had to aggregate cues then there is no unique ID.
@@ -648,11 +657,20 @@ void FElectraSubtitleDecoderWVTT::UpdatePlaybackPosition(Electra::FTimeValue InA
 		TSharedPtr<FSubtitleDecoderOutputWVTT, ESPMode::ThreadSafe> Out = MakeShared<FSubtitleDecoderOutputWVTT, ESPMode::ThreadSafe>();
 		Out->SetText(CombinedPlainText);
 		Electra::FTimeValue tv;
-		tv.SetFromTimespan(CueStartTime);
+		tv.SetFromTimespan(LargestCueStartTime);
 		tv.SetSequenceIndex(InAbsolutePosition.GetSequenceIndex());
 		Out->SetTimestamp(tv);
 		Electra::FTimeValue dur;
-		dur.SetFromTimespan(CueEndTime - CueStartTime);
+		// The duration of the subtitle we send is from the start time to the
+		// next evaluation time at which point we will construct a new aggegrated subtitle.
+		if (NextEval < FTimespan::MaxValue())
+		{
+			dur.SetFromTimespan(NextEval - LargestCueStartTime);
+		}
+		else
+		{
+			dur.SetFromSeconds(0.5);
+		}
 		Out->SetDuration(dur);
 		Out->SetID(CueID.Len() ? CueID : FString::Printf(TEXT("<%u>"), ++NextID), CueID.IsEmpty());
 
@@ -670,7 +688,7 @@ void FElectraSubtitleDecoderWVTT::UpdatePlaybackPosition(Electra::FTimeValue InA
 
 		if (bSendEmptySubtitleDuringGaps)
 		{
-			SendEmptySubtitleAt = CueEndTime;
+			SendEmptySubtitleAt = LargestCueEndTime;
 		}
 	}
 	else if (bSendEmptySubtitleDuringGaps && SendEmptySubtitleAt != FTimespan::MinValue())
diff --git a/Engine/Plugins/Media/ElectraSubtitles/Source/ElectraSubtitles/Private/wvtt/WebVTTParser.cpp b/Engine/Plugins/Media/ElectraSubtitles/Source/ElectraSubtitles/Private/wvtt/WebVTTParser.cpp
--- a/Engine/Plugins/Media/ElectraSubtitles/Source/ElectraSubtitles/Private/wvtt/WebVTTParser.cpp
+++ b/Engine/Plugins/Media/ElectraSubtitles/Source/ElectraSubtitles/Private/wvtt/WebVTTParser.cpp
@@ -1004,7 +1004,7 @@ void FWebVTTParser::GetCuesAtTime(TArray<TUniquePtr<ICue>>& OutCues, FTimespan&
 		}
 	}
 	// Look at the next cues starting at the iterator's position.
-	FTimespan NextChangeOccursAt = FTimespan::MaxValue();
+	TArray<FTimespan> CueTimes;
 	for(int32 i=It->CurrentIndex,iMax=CurrentDoc->Cues.Num(); i<iMax; ++i)
 	{
 		TRange<FTimespan> CueRange(CurrentDoc->Cues[i]->Start + CurrentDoc->TimestampOffset, CurrentDoc->Cues[i]->End + CurrentDoc->TimestampOffset);
@@ -1024,18 +1024,25 @@ void FWebVTTParser::GetCuesAtTime(TArray<TUniquePtr<ICue>>& OutCues, FTimespan&
 			{
 				OutCues.AddUnique(MoveTemp(Cue));
 			}
-			NextChangeOccursAt = FMath::Min(CueRange.GetUpperBoundValue(), NextChangeOccursAt);
+			// Add the cue range start and end to the cue times for determining the next
+			// time we need to get handled.
+			CueTimes.Add(CueRange.GetLowerBoundValue());
+			CueTimes.Add(CueRange.GetUpperBoundValue());
 		}
 		else if (CueRange.GetLowerBoundValue() > InAtTime)
 		{
-			if (OutCues.IsEmpty())
-			{
-				NextChangeOccursAt = CueRange.GetLowerBoundValue();
-			}
+			CueTimes.Add(CueRange.GetLowerBoundValue());
 			break;
 		}
 	}
-	OutNextChangeAt = NextChangeOccursAt;
+	// Sort the cue times and remove all those prior to the current time.
+	CueTimes.Sort();
+	while(!CueTimes.IsEmpty() && CueTimes[0] <= InAtTime)
+	{
+		CueTimes.RemoveAt(0);
+	}
+	// Then pick the next future time as the one we need to get handled again.
+	OutNextChangeAt = CueTimes.Num() ? CueTimes[0] : FTimespan::MaxValue();
 
 	// Skip the iterator ahead over all cues that are now in the past.
 	while(It->CurrentIndex < CurrentDoc->Cues.Num() && CurrentDoc->Cues[It->CurrentIndex]->End + CurrentDoc->TimestampOffset < InAtTime)
