ElectraPlayer: Fixing TTML subtitle timing in DASH period having a presentationTimeOffset; fixed index of selected track when playing across multiperiod DASH having different language availability per period.

#rnx

diff --git a/Engine/Plugins/Media/ElectraPlayer/Source/ElectraPlayerRuntime/Private/ElectraPlayer.cpp b/Engine/Plugins/Media/ElectraPlayer/Source/ElectraPlayerRuntime/Private/ElectraPlayer.cpp
--- a/Engine/Plugins/Media/ElectraPlayer/Source/ElectraPlayerRuntime/Private/ElectraPlayer.cpp
+++ b/Engine/Plugins/Media/ElectraPlayer/Source/ElectraPlayerRuntime/Private/ElectraPlayer.cpp
@@ -2212,10 +2212,7 @@ void FElectraPlayer::HandlePlayerEventTracksChanged()
 	bAudioTrackIndexDirty = true;
 	bSubtitleTrackIndexDirty = true;
 
-	/*
-	Note: Do not send TracksChanged here since this would trigger a re-selecting of the default tracks.
-		DeferredEvents.Enqueue(IElectraPlayerAdapterDelegate::EPlayerEvent::TracksChanged);
-	*/
+	DeferredEvents.Enqueue(IElectraPlayerAdapterDelegate::EPlayerEvent::TracksChanged);
 }
 
 
diff --git a/Engine/Plugins/Media/ElectraPlayer/Source/ElectraPlayerRuntime/Private/Runtime/Decoder/SubtitleDecoder.cpp b/Engine/Plugins/Media/ElectraPlayer/Source/ElectraPlayerRuntime/Private/Runtime/Decoder/SubtitleDecoder.cpp
--- a/Engine/Plugins/Media/ElectraPlayer/Source/ElectraPlayerRuntime/Private/Runtime/Decoder/SubtitleDecoder.cpp
+++ b/Engine/Plugins/Media/ElectraPlayer/Source/ElectraPlayerRuntime/Private/Runtime/Decoder/SubtitleDecoder.cpp
@@ -10,6 +10,7 @@ namespace Electra
 
 namespace SubtitleOptionKeys
 {
+	static const FName PresentationTimeOffset(TEXT("PresentationTimeOffset"));
 	static const FName SendEmptySubtitleDuringGaps(TEXT("sendEmptySubtitleDuringGaps"));
 	static const FName SourceID(TEXT("source_id"));
 	static const FName Width(TEXT("width"));
@@ -432,6 +433,10 @@ void FSubtitleDecoder::AUdataPushAU(FAccessUnit* AccessUnit)
 			{
 				Addtl.Set(SubtitleOptionKeys::SourceID, FVariantValue(AccessUnit->BufferSourceInfo->PeriodAdaptationSetID));
 			}
+			if (AccessUnit->AUCodecData.IsValid() && AccessUnit->AUCodecData->ParsedInfo.GetExtras().HaveKey(StreamCodecInformationOptions::PresentationTimeOffset))
+			{
+				Addtl.Set(SubtitleOptionKeys::PresentationTimeOffset, AccessUnit->AUCodecData->ParsedInfo.GetExtras().GetValue(StreamCodecInformationOptions::PresentationTimeOffset));
+			}
 			PluginDecoder->AddStreamedSubtitleData(TArray<uint8>((const uint8*)AccessUnit->AUData, (int32)AccessUnit->AUSize), AccessUnit->PTS, AccessUnit->Duration, Addtl);
 		}
 	}
diff --git a/Engine/Plugins/Media/ElectraPlayer/Source/ElectraPlayerRuntime/Private/Runtime/Player/StreamSegmentReaderCommon.cpp b/Engine/Plugins/Media/ElectraPlayer/Source/ElectraPlayerRuntime/Private/Runtime/Player/StreamSegmentReaderCommon.cpp
--- a/Engine/Plugins/Media/ElectraPlayer/Source/ElectraPlayerRuntime/Private/Runtime/Player/StreamSegmentReaderCommon.cpp
+++ b/Engine/Plugins/Media/ElectraPlayer/Source/ElectraPlayerRuntime/Private/Runtime/Player/StreamSegmentReaderCommon.cpp
@@ -1795,6 +1795,9 @@ FStreamSegmentReaderCommon::FStreamHandler::EHandleResult FStreamSegmentReaderCo
 								LogMessage(IInfoLog::ELevel::Verbose, FString::Printf(TEXT("Track timescale %u differs from timescale of %u in MPD or segment index. This may cause playback problems!"), TrackTimescale, InRequest->Segment.Timescale));
 							}
 						}
+						// Set the PTO in the codec data extras. This is a rarely used value and constant for the segment.
+						// The only use we have for it so far is to remap subtitle timestamps to split periods.
+						CurrentlyActiveTrackData->CSD->ParsedInfo.GetExtras().Set(StreamCodecInformationOptions::PresentationTimeOffset, FVariantValue(FTimeValue(CurrentlyActiveTrackData->PTO, TrackTimescale)));
 
 						// Producer reference time (DASH only)
 						if (InRequest->Segment.MeasureLatencyViaReferenceTimeInfoID >= 0 && InRequest->Segment.ProducerReferenceTimeInfos.Num())
@@ -3489,6 +3492,10 @@ FStreamSegmentReaderCommon::FStreamHandler::EHandleResult FStreamSegmentReaderCo
 		}
 		LatestPTS.SetSequenceIndex(InRequest->TimestampSequenceIndex);
 
+		// Set the PTO in the codec data extras. This is a rarely used value and constant for the segment.
+		// The only use we have for it so far is to remap subtitle timestamps to split periods.
+		CurrentlyActiveTrackData->CSD->ParsedInfo.GetExtras().Set(StreamCodecInformationOptions::PresentationTimeOffset, FVariantValue(PTO));
+
 		// Create an access unit
 		FAccessUnit *AccessUnit = FAccessUnit::Create(Parameters.MemoryProvider);
 		AccessUnit->ESType = InRequest->GetType();
diff --git a/Engine/Plugins/Media/ElectraPlayer/Source/ElectraPlayerRuntime/Private/Runtime/StreamTypes.h b/Engine/Plugins/Media/ElectraPlayer/Source/ElectraPlayerRuntime/Private/Runtime/StreamTypes.h
--- a/Engine/Plugins/Media/ElectraPlayer/Source/ElectraPlayerRuntime/Private/Runtime/StreamTypes.h
+++ b/Engine/Plugins/Media/ElectraPlayer/Source/ElectraPlayerRuntime/Private/Runtime/StreamTypes.h
@@ -716,6 +716,7 @@ namespace Electra
 		static const FName SampleSize(TEXT("SampleSize"));
 		static const FName ConstBytesPerAudioPacket(TEXT("ConstBytesPerAudioPacket"));
 		static const FName ConstLPCMFramesPerAudioPacket(TEXT("ConstLPCMFramesPerAudioPacket"));
+		static const FName PresentationTimeOffset(TEXT("PresentationTimeOffset"));
 	}
 
 
diff --git a/Engine/Plugins/Media/ElectraSubtitles/Source/ElectraSubtitles/Private/ttml/ElectraSubtitleDecoder_TTML.cpp b/Engine/Plugins/Media/ElectraSubtitles/Source/ElectraSubtitles/Private/ttml/ElectraSubtitleDecoder_TTML.cpp
--- a/Engine/Plugins/Media/ElectraSubtitles/Source/ElectraSubtitles/Private/ttml/ElectraSubtitleDecoder_TTML.cpp
+++ b/Engine/Plugins/Media/ElectraSubtitles/Source/ElectraSubtitles/Private/ttml/ElectraSubtitleDecoder_TTML.cpp
@@ -110,9 +110,9 @@ private:
 
 namespace ElectraSubtitleDecoderTTMLOptions
 {
+	static const FName PresentationTimeOffset(TEXT("PresentationTimeOffset"));
 	static const FName SendEmptySubtitleDuringGaps(TEXT("sendEmptySubtitleDuringGaps"));
 	static const FName SourceID(TEXT("source_id"));
-	static const FName PTO(TEXT("PTO"));
 	static const FName Width(TEXT("width"));
 	static const FName Height(TEXT("height"));
 	static const FName OffsetX(TEXT("offset_x"));
@@ -189,7 +189,7 @@ void FElectraSubtitleDecoderTTML::AddStreamedSubtitleData(const TArray<uint8>& I
 			NewTimelineHandler.Handler = SubtitleHandler;
 			NewTimelineHandler.AbsoluteTimeRange.Start = InAbsoluteTimestamp;
 			NewTimelineHandler.AbsoluteTimeRange.End = InAbsoluteTimestamp + InDuration;
-			NewTimelineHandler.MediaLocalPTO = InAdditionalInfo.GetValue(ElectraSubtitleDecoderTTMLOptions::PTO).SafeGetTimeValue(Electra::FTimeValue::GetZero());
+			NewTimelineHandler.MediaLocalPTO = InAdditionalInfo.GetValue(ElectraSubtitleDecoderTTMLOptions::PresentationTimeOffset).SafeGetTimeValue(Electra::FTimeValue::GetZero());
 			TimelineHandlers.Emplace(MoveTemp(NewTimelineHandler));
 		}
 	}
@@ -226,7 +226,7 @@ void FElectraSubtitleDecoderTTML::AddStreamedSubtitleData(const TArray<uint8>& I
 				NewTimelineHandler.Handler = SubtitleHandler;
 				NewTimelineHandler.AbsoluteTimeRange.Start = InAbsoluteTimestamp;
 				NewTimelineHandler.AbsoluteTimeRange.End = InAbsoluteTimestamp + InDuration;
-				NewTimelineHandler.MediaLocalPTO = InAdditionalInfo.GetValue(ElectraSubtitleDecoderTTMLOptions::PTO).SafeGetTimeValue(Electra::FTimeValue::GetZero());
+				NewTimelineHandler.MediaLocalPTO = InAdditionalInfo.GetValue(ElectraSubtitleDecoderTTMLOptions::PresentationTimeOffset).SafeGetTimeValue(Electra::FTimeValue::GetZero());
 				TimelineHandlers.Emplace(MoveTemp(NewTimelineHandler));
 				TimelineHandlers.StableSort([](const FTimelineHandlers& e1, const FTimelineHandlers& e2)
 				{
