ElectraSubtitles: Fixed an incorrectly released mutex

#rnx

diff --git a/Engine/Plugins/Media/ElectraSubtitles/Source/ElectraSubtitles/Private/wvtt/ElectraSubtitleDecoder_WVTT.cpp b/Engine/Plugins/Media/ElectraSubtitles/Source/ElectraSubtitles/Private/wvtt/ElectraSubtitleDecoder_WVTT.cpp
--- a/Engine/Plugins/Media/ElectraSubtitles/Source/ElectraSubtitles/Private/wvtt/ElectraSubtitleDecoder_WVTT.cpp
+++ b/Engine/Plugins/Media/ElectraSubtitles/Source/ElectraSubtitles/Private/wvtt/ElectraSubtitleDecoder_WVTT.cpp
@@ -371,7 +371,7 @@ void FElectraSubtitleDecoderWVTT::AddStreamedSubtitleData(const TArray<uint8>& I
 				return;
 			}
 		}
-		ParsedTimerangeLock.Unlock();
+		lock.Unlock();
 
 		TSharedPtr<FParsedTimerange, ESPMode::ThreadSafe> ParsedTimerange = MakeShared<FParsedTimerange, ESPMode::ThreadSafe>();
 		ParsedTimerange->Parser = ElectraWebVTTParser::IWebVTTParser::Create();
