[ElectraSubtitles] Include CoreMinimal.h from TTMLSubtitleHandler.h because UE Tests is having compile errors.

diff --git a/Engine/Plugins/Media/ElectraSubtitles/Source/ElectraSubtitles/Private/ttml/TTMLSubtitleHandler.h b/Engine/Plugins/Media/ElectraSubtitles/Source/ElectraSubtitles/Private/ttml/TTMLSubtitleHandler.h
--- a/Engine/Plugins/Media/ElectraSubtitles/Source/ElectraSubtitles/Private/ttml/TTMLSubtitleHandler.h
+++ b/Engine/Plugins/Media/ElectraSubtitles/Source/ElectraSubtitles/Private/ttml/TTMLSubtitleHandler.h
@@ -2,6 +2,7 @@
 
 #pragma once
 
+#include "CoreMinimal.h"
 #include "PlayerTime.h"
 
 struct FScriptContainerElement;
